Copying //UE4/Dev-Networking to //UE4/Dev-Main (Source: //UE4/Dev-Networking @ 4051526)

diff --git a/Engine/Plugins/Runtime/ReplicationGraph/Source/Private/ReplicationGraph.cpp b/Engine/Plugins/Runtime/ReplicationGraph/Source/Private/ReplicationGraph.cpp
new file mode 100644
--- /dev/null
+++ b/Engine/Plugins/Runtime/ReplicationGraph/Source/Private/ReplicationGraph.cpp
@@ -0,0 +1,3208 @@
+// Copyright 1998-2018 Epic Games, Inc. All Rights Reserved.
+// 
+
+
+/**
+ *	
+ *	===================== TODO / WIP Notes =====================
+ *
+ * TODO Missing Features:
+ *	-bNetTemporary
+ *	
+ *	 
+ * TODO Refactors/cleanup:	 
+ *	-Packet Budget: simplify or possibly remove. Maybe split into two concepts: a bitmask for filtering what gets pulled out of the graph and how traffic is tracked as it is sent.
+ *
+ *	-FActorRepList/FPrioritizedRepList/FActorRepListRefView: Probably can be replaced with regular TArrays/TArrayViews. The original problem these were solving were having lots of little
+ *	lists that needed to be packed and ref counted. Neither is really needed anymore.
+ *	
+ *	
+ * TODO Optimizations:	
+ *	-Need to store stable pointers to FReplicationActorInfo/FGlobalActorReplicationInfo. See FPrioritizedRepList::GlobalData/ConnectionData. Want to make it so these are only ever looked up once per frame at most.
+ *	 	
+ *
+ *	 	
+ *	 	
+ *	--------------------------------
+ *	
+ *	Game Code API
+ *	
+ *	
+ *	Function						Status (w/ RepDriver enabled)
+ *	----------------------------------------------------------------------------------------
+ *	ForceNetUpdate					Compatible/Working			
+ *	FlushNetDormancy				Compatible/Working
+ *	SetNetUpdateTime				NOOP							
+ *	ForceNetRelevant				NOOP
+ *	ForceActorRelevantNextUpdate	NOOP
+ *
+ *	FindOrAddNetworkObjectInfo		NOOP. Accessing legacy system data directly. This sucks and should never have been exposed to game code directly.
+ *	FindNetworkObjectInfo			NOOP. Will want to deprecate both of these functions and get them out of our code base
+ *	
+ *	
+ *
+ */
+
+#include "ReplicationGraph.h"
+#include "EngineGlobals.h"
+#include "Engine/World.h"
+
+#include "Engine/LocalPlayer.h"
+#include "EngineUtils.h"
+#include "Engine/Engine.h"
+#include "Net/DataReplication.h"
+#include "Engine/ActorChannel.h"
+#include "Engine/NetworkObjectList.h"
+#include "Net/RepLayout.h"
+#include "GameFramework/SpectatorPawn.h"
+#include "Engine/LevelStreamingKismet.h"
+#include "GameFramework/SpectatorPawnMovement.h"
+#include "Net/UnrealNetwork.h"
+#include "Net/NetworkProfiler.h"
+#include "HAL/LowLevelMemTracker.h"
+#include "UObject/UObjectIterator.h"
+#include "Engine/Level.h"
+#include "Templates/UnrealTemplate.h"
+#include "Stats/StatsMisc.h"
+#include "Net/DataChannel.h"
+#include "UObject/UObjectGlobals.h"
+
+int32 CVar_RepGraph_Pause = 0;
+static FAutoConsoleVariableRef CVarRepGraphPause(TEXT("Net.RepGraph.Pause"), CVar_RepGraph_Pause, TEXT("Pauses actor replication in the Replication Graph."), ECVF_Default );
+
+int32 CVar_RepGraph_Frequency = 0.00f;
+static FAutoConsoleVariableRef CVarRepGraphFrequency(TEXT("Net.RepGraph.Frequency"), CVar_RepGraph_Frequency, TEXT("Enabled Replication Manager. 0 will fallback to legacy NetDriver implementation."), ECVF_Default );
+
+int32 CVar_RepGraph_UseLegacyBudget = 1;
+static FAutoConsoleVariableRef CVarRepGraphUseLegacyBudget(TEXT("Net.RepGraph.UseLegacyBudget"), CVar_RepGraph_UseLegacyBudget, TEXT("Use legacy IsNetReady() to make dynamic packget budgets"), ECVF_Default );
+
+float CVar_RepGraph_FixedBudget = 0;
+static FAutoConsoleVariableRef CVarRepGraphFixedBudge(TEXT("Net.RepGraph.FixedBudget"), CVar_RepGraph_FixedBudget, TEXT("Set fixed (independent of frame rate) packet budget. In BIts/frame"), ECVF_Default );
+
+int32 CVar_RepGraph_SkipDistanceCull = 0;
+static FAutoConsoleVariableRef CVarRepGraphSkipDistanceCull(TEXT("Net.RepGraph.Skip Distance Cull"), CVar_RepGraph_SkipDistanceCull, TEXT(""), ECVF_Default );
+
+int32 CVar_RepGraph_PrintCulledOnConnectionClasses = 0;
+static FAutoConsoleVariableRef CVarRepGraphPrintCulledOnConnectionClasses(TEXT("Net.RepGraph.PrintCulledOnConnectionClasses"), CVar_RepGraph_PrintCulledOnConnectionClasses, TEXT(""), ECVF_Default );
+
+int32 CVar_RepGraph_TrackClassReplication = 0;
+static FAutoConsoleVariableRef CVarRepGraphTrackClassReplication(TEXT("Net.RepGraph.TrackClassReplication"), CVar_RepGraph_TrackClassReplication, TEXT(""), ECVF_Default );
+
+int32 CVar_RepGraph_PrintTrackClassReplication = 0;
+static FAutoConsoleVariableRef CVarRepGraphPrintTrackClassReplication(TEXT("Net.RepGraph.PrintTrackClassReplication"), CVar_RepGraph_PrintTrackClassReplication, TEXT(""), ECVF_Default );
+
+REPGRAPH_DEVCVAR_SHIPCONST(int32, "Net.RepGraph.LogNetDormancyDetails", CVar_RepGraph_LogNetDormancyDetails, 0, "Logs actors that are removed from the replication graph/nodes.");
+REPGRAPH_DEVCVAR_SHIPCONST(int32, "Net.RepGraph.LogActorRemove", CVar_RepGraph_LogActorRemove, 0, "Logs actors that are removed from the replication graph/nodes.");
+REPGRAPH_DEVCVAR_SHIPCONST(int32, "Net.RepGraph.LogActorAdd", CVar_RepGraph_LogActorAdd, 0, "Logs actors that are added to replication graph/nodes.");
+REPGRAPH_DEVCVAR_SHIPCONST(int32, "Net.RepGraph.Verify", CVar_RepGraph_Verify, 0, "Additional, slow, verification is done on replication graph nodes. Guards against: invalid actors and dupes");
+
+REPGRAPH_DEVCVAR_SHIPCONST(int32, "Net.RepGraph.TrickleDistCullOnDormanyNodes", CVar_RepGraph_TrickleDistCullOnDormanyNodes, 1, "Actors in a dormancy node that are distance culled will trickle through as dormancy node empties");
+REPGRAPH_DEVCVAR_SHIPCONST(int32, "Net.RepGraph.EnableRPCSendPolicy", CVar_RepGraph_EnableRPCSendPolicy, 1, "Enables RPC send policy (e.g, force certain functions to send immediately rather than be queued)");
+
+DECLARE_STATS_GROUP(TEXT("ReplicationDriver"), STATGROUP_RepDriver, STATCAT_Advanced);
+DECLARE_DWORD_COUNTER_STAT(TEXT("Rep Actor List Dupes"), STAT_NetRepActorListDupes, STATGROUP_RepDriver);
+DECLARE_DWORD_COUNTER_STAT(TEXT("Actor Channels Opened"), STAT_NetActorChannelsOpened, STATGROUP_RepDriver);
+DECLARE_DWORD_COUNTER_STAT(TEXT("Actor Channels Closed"), STAT_NetActorChannelsClosed, STATGROUP_RepDriver);
+DECLARE_DWORD_COUNTER_STAT(TEXT("Num Processed Connections"), STAT_NumProcessedConnections, STATGROUP_RepDriver);
+
+// --------------------------------------------------------------------------------------------------------------------------------------------
+// --------------------------------------------------------------------------------------------------------------------------------------------
+// --------------------------------------------------------------------------------------------------------------------------------------------
+
+
+UReplicationGraph::UReplicationGraph()
+{
+	ReplicationConnectionManagerClass = UNetReplicationGraphConnection::StaticClass();
+	GlobalActorChannelFrameNumTimeout = 2;
+
+	if (!HasAnyFlags(RF_ClassDefaultObject))
+	{
+		GraphGlobals = MakeShared<FReplicationGraphGlobalData>();
+		GraphGlobals->GlobalActorReplicationInfoMap = &GlobalActorReplicationInfoMap;
+	}
+}
+
+void UReplicationGraph::InitForNetDriver(UNetDriver* InNetDriver)
+{
+	NetDriver = InNetDriver;
+
+	InitGlobalActorClassSettings();
+	InitGlobalGraphNodes();
+
+	for (UNetConnection* ClientConnection : NetDriver->ClientConnections)
+	{
+		AddClientConnection(ClientConnection);
+	}
+}
+
+void UReplicationGraph::InitNode(UReplicationGraphNode* Node)
+{
+	Node->Initialize(GraphGlobals);
+
+	if (Node->GetRequiresPrepareForReplication())
+	{
+		PrepareForReplicationNodes.Add(Node);
+	}
+}
+
+void UReplicationGraph::InitGlobalActorClassSettings()
+{
+	// AInfo and APlayerControllers have no world location, so distance scaling should always be 0
+	FClassReplicationInfo NonSpatialClassInfo;
+	NonSpatialClassInfo.DistancePriorityScale = 0.f;
+
+	GlobalActorReplicationInfoMap.SetClassInfo( AInfo::StaticClass(), NonSpatialClassInfo );
+	GlobalActorReplicationInfoMap.SetClassInfo( APlayerController::StaticClass(), NonSpatialClassInfo );
+}
+
+void UReplicationGraph::InitGlobalGraphNodes()
+{
+	// TODO: We should come up with a basic/default implementation for people to use to model
+}
+
+void UReplicationGraph::InitConnectionGraphNodes(UNetReplicationGraphConnection* ConnectionManager)
+{
+	// This handles tear off actors. Child classes should call Super::InitConnectionGraphNodes.
+	ConnectionManager->TearOffNode = CreateNewNode<UReplicationGraphNode_TearOff_ForConnection>();
+	ConnectionManager->AddConnectionGraphNode(ConnectionManager->TearOffNode);
+}
+
+void UReplicationGraph::AddGlobalGraphNode(UReplicationGraphNode* GraphNode)
+{
+	GlobalGraphNodes.Add(GraphNode);
+}
+
+void UReplicationGraph::AddConnectionGraphNode(UReplicationGraphNode* GraphNode, UNetReplicationGraphConnection* ConnectionManager)
+{
+	ConnectionManager->AddConnectionGraphNode(GraphNode);
+}
+
+void UReplicationGraph::RemoveConnectionGraphNode(UReplicationGraphNode* GraphNode, UNetReplicationGraphConnection* ConnectionManager)
+{
+	ConnectionManager->RemoveConnectionGraphNode(GraphNode);
+}
+
+UNetReplicationGraphConnection* UReplicationGraph::FindOrAddConnectionManager(UNetConnection* NetConnection)
+{
+	FScopeLogTime SLT( TEXT( "UReplicationGraph::FindOrAddConnectionManager(" ), nullptr, FScopeLogTime::ScopeLog_Milliseconds );
+
+	// Could use an acceleration map if necessary
+	RG_QUICK_SCOPE_CYCLE_COUNTER(UReplicationGraph_FindConnectionManager)
+	for (UNetReplicationGraphConnection* ConnManager : Connections)
+	{
+		if (ConnManager->NetConnection == NetConnection)
+		{
+			return ConnManager;
+		}
+	}
+
+	for (UNetReplicationGraphConnection* ConnManager : PendingConnections)
+	{
+		if (ConnManager->NetConnection == NetConnection)
+		{
+			return ConnManager;
+		}
+	}
+
+	// We dont have one yet, create one but put it in the pending list. ::AddClientConnection *should* be called soon!
+	UNetReplicationGraphConnection* NewManager = CreateClientConnectionManagerInternal(NetConnection);
+	PendingConnections.Add(NewManager);
+	return NewManager;
+}
+
+void UReplicationGraph::AddClientConnection(UNetConnection* NetConnection)
+{
+	FScopeLogTime SLT( TEXT( "UReplicationGraph::AddClientConnection" ), nullptr, FScopeLogTime::ScopeLog_Milliseconds );
+
+	// We may have already created a manager for this connection in the pending list
+	for (int32 i=PendingConnections.Num()-1; i >= 0; --i)
+	{
+		if (UNetReplicationGraphConnection* ConnManager = PendingConnections[i])
+		{
+			if (ConnManager->NetConnection == NetConnection)
+			{
+				PendingConnections.RemoveAtSwap(i, 1, false);
+				Connections.Add(ConnManager);
+				return;
+			}
+		}
+	}
+
+	// Create it
+	Connections.Add(CreateClientConnectionManagerInternal(NetConnection));
+}
+
+UNetReplicationGraphConnection* UReplicationGraph::CreateClientConnectionManagerInternal(UNetConnection* Connection)
+{
+	repCheckf(Connection->GetReplicationConnectionDriver() == nullptr, TEXT("Connection %s on NetDriver %s already has a ReplicationConnectionDriver %s"), *GetNameSafe(Connection), *GetNameSafe(Connection->Driver), *Connection->GetReplicationConnectionDriver()->GetName() );
+
+	// Create the object
+	UNetReplicationGraphConnection* NewConnectionManager = NewObject<UNetReplicationGraphConnection>(this, ReplicationConnectionManagerClass.Get());
+
+	// Give it an ID
+	NewConnectionManager->ConnectionId = Connections.Num() + PendingConnections.Num();
+
+	// Initialize it with us
+	NewConnectionManager->InitForGraph(this);
+
+	// Associate NetConnection with it
+	NewConnectionManager->InitForConnection(Connection);
+
+	// Create Graph Nodes for this specific connection
+	InitConnectionGraphNodes(NewConnectionManager);
+
+	return NewConnectionManager;
+}
+
+void UReplicationGraph::RemoveClientConnection(UNetConnection* NetConnection)
+{
+	int32 ConnectionId = 0;
+	bool bFound = false;
+
+	// Remove the RepGraphConnection associated with this NetConnection. Also update ConnectionIds to stay compact.
+	auto UpdateList = [&](TArray<UNetReplicationGraphConnection*> List)
+	{
+		for (int32 idx=0; idx < Connections.Num(); ++idx)
+		{
+			UNetReplicationGraphConnection* ConnectionManager = Connections[idx];
+			repCheck(ConnectionManager);
+
+			if (ConnectionManager->NetConnection == NetConnection)
+			{
+				ensure(!bFound);
+				Connections.RemoveAtSwap(idx, 1, false);
+				bFound = true;
+			}
+			else
+			{
+				ConnectionManager->ConnectionId = ConnectionId++;
+			}
+		}
+	};
+
+	UpdateList(Connections);
+	UpdateList(PendingConnections);
+
+	if (!bFound)
+	{
+		// At least one list should have found the connection
+		UE_LOG(LogReplicationGraph, Warning, TEXT("UReplicationGraph::RemoveClientConnection could not find connection in Connection (%d) or PendingConnections (%d) lists"), *GetNameSafe(NetConnection), Connections.Num(), PendingConnections.Num());
+	}
+}
+
+void UReplicationGraph::SetWorld(UWorld* InWorld)
+{
+	if (GraphGlobals.IsValid())
+	{
+		GraphGlobals->World = InWorld;
+	}
+
+	if (InWorld)
+	{
+		if (InWorld->AreActorsInitialized())
+		{
+			InitializeForWorld(InWorld);
+		}
+		else
+		{
+			// World isn't initialized yet. This happens when launching into a map directly from command line
+			InWorld->OnActorsInitialized.AddLambda([&](const UWorld::FActorsInitializedParams& P)
+			{
+				this->InitializeForWorld(P.World);
+			});
+		}
+	}
+}
+
+void UReplicationGraph::InitializeForWorld(UWorld* World)
+{
+	ActiveNetworkActors.Reset();
+
+	for (UReplicationGraphNode* Manager : GlobalGraphNodes)
+	{
+		Manager->NotifyResetAllNetworkActors();
+	}
+	
+	if (World)
+	{
+		for (FActorIterator Iter(World); Iter; ++Iter)
+		{
+			AActor* Actor = *Iter;
+			if (Actor != nullptr && !Actor->IsPendingKill() && ULevel::IsNetActor(Actor))
+			{
+				AddNetworkActor(Actor);
+			}
+		}
+	}
+}
+
+void UReplicationGraph::AddNetworkActor(AActor* Actor)
+{
+	RG_QUICK_SCOPE_CYCLE_COUNTER(UReplicationGraph_AddNetworkActor);
+
+	if (IsActorValidForReplicationGather(Actor) == false)
+	{
+		return;
+	}
+
+	bool bWasAlreadyThere = false;
+	ActiveNetworkActors.Add(Actor, &bWasAlreadyThere);
+	if (bWasAlreadyThere)
+	{
+		// Guarding against double adds
+		return;
+	}
+
+	// Create global rep info	
+	FGlobalActorReplicationInfo& GlobalInfo = GlobalActorReplicationInfoMap.Get(Actor);
+	GlobalInfo.bWantsToBeDormant = Actor->NetDormancy > DORM_Awake;
+
+	RouteAddNetworkActorToNodes(FNewReplicatedActorInfo(Actor), GlobalInfo);
+}
+
+
+void UReplicationGraph::RouteAddNetworkActorToNodes(const FNewReplicatedActorInfo& ActorInfo, FGlobalActorReplicationInfo& GlobalInfo)
+{
+	// The base implementation just routes to every global node. Subclasses will want a more direct routing function where possible.
+	for (UReplicationGraphNode* Node : GlobalGraphNodes)
+	{
+		Node->NotifyAddNetworkActor(ActorInfo);
+	}
+}
+
+void UReplicationGraph::RemoveNetworkActor(AActor* Actor)
+{
+	RG_QUICK_SCOPE_CYCLE_COUNTER(UReplicationGraph_RemoveNetworkActor);
+
+	if (ActiveNetworkActors.Remove(Actor) == 0)
+	{
+		// Guarding against double removes
+		return;
+	}
+
+	// Tear off actors have already been removed from the nodes, so we don't need to route them again.
+	if (Actor->GetTearOff() == false)
+	{
+		UE_CLOG(CVar_RepGraph_LogActorRemove>0, LogReplicationGraph, Display, TEXT("UReplicationGraph::RemoveNetworkActor %s"), *Actor->GetFullName());
+		RouteRemoveNetworkActorToNodes(FNewReplicatedActorInfo(Actor));
+	}
+
+	GlobalActorReplicationInfoMap.Remove(Actor);
+}
+
+void UReplicationGraph::RouteRemoveNetworkActorToNodes(const FNewReplicatedActorInfo& ActorInfo)
+{
+	// The base implementation just routes to every global node. Subclasses will want a more direct routing function where possible.
+	for (UReplicationGraphNode* Node : GlobalGraphNodes)
+	{
+		Node->NotifyRemoveNetworkActor(ActorInfo);
+	}
+}
+
+void UReplicationGraph::ForceNetUpdate(AActor* Actor)
+{
+	if (FGlobalActorReplicationInfo* RepInfo = GlobalActorReplicationInfoMap.Find(Actor))
+	{
+		RepInfo->ForceNetUpdateFrame = ReplicationGraphFrame;
+		RepInfo->Events.ForceNetUpdate.Broadcast(Actor, *RepInfo);
+	}
+}
+
+void UReplicationGraph::FlushNetDormancy(AActor* Actor, bool bWasDormInitial)
+{
+	RG_QUICK_SCOPE_CYCLE_COUNTER(UReplicationGraph_FlushNetDormancy);
+
+	FGlobalActorReplicationInfo& GlobalInfo = GlobalActorReplicationInfoMap.Get(Actor);
+	const bool bNewWantsToBeDormant = (Actor->NetDormancy > DORM_Awake);
+
+	UE_CLOG(CVar_RepGraph_LogNetDormancyDetails > 0, LogReplicationGraph, Display, TEXT("UReplicationGraph::FlushNetDormancy %s. Old WantsToBeDormant: %d. New WantsToBeDormant: %d"), *Actor->GetPathName(), GlobalInfo.bWantsToBeDormant, bNewWantsToBeDormant);
+
+	if (GlobalInfo.bWantsToBeDormant != bNewWantsToBeDormant)
+	{
+		UE_LOG(LogReplicationGraph, Display, TEXT("UReplicationGraph::FlushNetDormancy %s. WantsToBeDormant is changing (%d -> %d) from a Flush! We expect NotifyActorDormancyChange to be called first."), *Actor->GetPathName(), GlobalInfo.bWantsToBeDormant, bNewWantsToBeDormant);
+		GlobalInfo.bWantsToBeDormant = Actor->NetDormancy > DORM_Awake;
+	}
+
+	if (GlobalInfo.bWantsToBeDormant == false)
+	{
+		// This actor doesn't want to be dormant. Suppress the Flush call into the nodes. This is to prevent wasted work since the AActor code calls NotifyActorDormancyChange then Flush always.
+		return;
+	}
+
+	if (bWasDormInitial)
+	{
+		AddNetworkActor(Actor);
+	}
+	else
+	{
+		GlobalInfo.Events.DormancyFlush.Broadcast(Actor, GlobalInfo);
+
+		// Stinks to have to iterate through like this, especially when net driver is doing a similar thing.
+		// Dormancy should probably be rewritten.
+		for (UNetReplicationGraphConnection* ConnectionManager: Connections)
+		{
+			if (FConnectionReplicationActorInfo* Info = ConnectionManager->ActorInfoMap.Find(Actor))
+			{
+				Info->bDormantOnConnection = false;
+			}
+		}
+	}
+}
+
+void UReplicationGraph::NotifyActorTearOff(AActor* Actor)
+{
+	// All connections that currently have a channel for the actor will put this actor on their TearOffNode.
+	for (UNetReplicationGraphConnection* ConnectionManager: Connections)
+	{
+		if (FConnectionReplicationActorInfo* Info = ConnectionManager->ActorInfoMap.Find(Actor))
+		{
+			UActorChannel* Channel = Info->Channel;
+			if (Channel && Channel->Actor)
+			{
+				Info->bTearOff = true; // Tells ServerReplicateActors to close the channel the next time this replicates
+				ConnectionManager->TearOffNode->NotifyTearOffActor(Actor, Info->LastRepFrameNum); // Tells this connection to gather this actor (until it replicates again)
+			}
+		}
+	}
+
+	// Remove the actor from the rest of the graph. The tear off node will add it from here.
+	RouteRemoveNetworkActorToNodes(FNewReplicatedActorInfo(Actor));
+}
+
+void UReplicationGraph::NotifyActorFullyDormantForConnection(AActor* Actor, UNetConnection* Connection)
+{
+	RG_QUICK_SCOPE_CYCLE_COUNTER(UReplicationGraph_NotifyActorFullyDormantForConnection);
+
+	// This is kind of bad but unavoidable. Possibly could use acceleration map (actor -> connections) but that would be a pain to maintain.
+	for (UNetReplicationGraphConnection* ConnectionManager: Connections)
+	{
+		if (ConnectionManager->NetConnection == Connection)
+		{			
+			if (FConnectionReplicationActorInfo* Info = ConnectionManager->ActorInfoMap.Find(Actor))
+			{
+				Info->bDormantOnConnection = true;
+			}
+			break;
+		}
+	}
+}
+
+void UReplicationGraph::NotifyActorDormancyChange(AActor* Actor, ENetDormancy OldDormancyState)
+{
+	RG_QUICK_SCOPE_CYCLE_COUNTER(UReplicationGraph_NotifyActorDormancyChange);
+
+	FGlobalActorReplicationInfo& GlobalInfo = GlobalActorReplicationInfoMap.Get(Actor);
+	ENetDormancy CurrentDormancy = Actor->NetDormancy;
+
+	UE_CLOG(CVar_RepGraph_LogNetDormancyDetails > 0, LogReplicationGraph, Display, TEXT("UReplicationGraph::NotifyActorDormancyChange %s. Old WantsToBeDormant: %d. New WantsToBeDormant: %d"), *Actor->GetPathName(), GlobalInfo.bWantsToBeDormant, CurrentDormancy > DORM_Awake ? 1 : 0);
+
+	GlobalInfo.bWantsToBeDormant = CurrentDormancy > DORM_Awake;
+	GlobalInfo.Events.DormancyChange.Broadcast(Actor, GlobalInfo, CurrentDormancy, OldDormancyState);
+}
+
+FORCEINLINE bool ReadyForNextReplication(FConnectionReplicationActorInfo& ConnectionData, FGlobalActorReplicationInfo& GlobalData, const uint32 FrameNum)
+{
+	return (ConnectionData.NextReplicationFrameNum <= FrameNum || GlobalData.ForceNetUpdateFrame > ConnectionData.LastRepFrameNum);
+}
+
+// ----------------------------------------------------------------------------------------------------------------------------------------------------------------------------------------------------------------------------------------------
+// ----------------------------------------------------------------------------------------------------------------------------------------------------------------------------------------------------------------------------------------------
+// ----------------------------------------------------------------------------------------------------------------------------------------------------------------------------------------------------------------------------------------------
+//	Server Replicate Actors
+// ----------------------------------------------------------------------------------------------------------------------------------------------------------------------------------------------------------------------------------------------
+// ----------------------------------------------------------------------------------------------------------------------------------------------------------------------------------------------------------------------------------------------
+// ----------------------------------------------------------------------------------------------------------------------------------------------------------------------------------------------------------------------------------------------
+
+FNativeClassAccumulator ChangeClassAccumulator;
+FNativeClassAccumulator NoChangeClassAccumulator;
+
+#if !(UE_BUILD_SHIPPING || UE_BUILD_TEST)
+	bool bTrackClassReplication = false;
+#else
+	const bool bTrackClassReplication = false;
+#endif
+
+int32 UReplicationGraph::ServerReplicateActors(float DeltaSeconds)
+{
+#if !(UE_BUILD_SHIPPING || UE_BUILD_TEST)
+	if (CVar_RepGraph_Pause)
+	{
+		return 0;
+	}
+
+	// Temp Hack for frequency throttling
+	static float TimeLeft = CVar_RepGraph_Frequency;
+	TimeLeft -= DeltaSeconds;
+	if (TimeLeft > 0.f)
+	{
+		return 0;
+	}
+	TimeLeft = CVar_RepGraph_Frequency;
+#endif
+	
+	++NetDriver->ReplicationFrame;	// This counter is used by RepLayout to utilize CL/serialization sharing. We must increment it ourselves, but other places can increment it too, in order to invalidate the shared state.
+	const uint32 FrameNum = ++ReplicationGraphFrame; // This counter is used internally and drives all frame based replication logic.
+
+	// -------------------------------------------------------
+	//	PREPARE (Global)
+	// -------------------------------------------------------
+
+	{
+		RG_QUICK_SCOPE_CYCLE_COUNTER(NET_PrepareReplication);
+
+		for (UReplicationGraphNode* Node : PrepareForReplicationNodes)
+		{
+			Node->PrepareForReplication();
+		}
+	}
+	
+#if !(UE_BUILD_SHIPPING || UE_BUILD_TEST)
+	const bool bDistanceCull = (CVar_RepGraph_SkipDistanceCull == 0);
+	const bool bCulledOnConnectionCount = (CVar_RepGraph_PrintCulledOnConnectionClasses == 1);	
+	bTrackClassReplication = (CVar_RepGraph_TrackClassReplication > 0 || CVar_RepGraph_PrintTrackClassReplication > 0);
+	if (!bTrackClassReplication)
+	{
+		ChangeClassAccumulator.Reset();
+		NoChangeClassAccumulator.Reset();
+	}
+
+#else
+	const bool bDistanceCull = true;
+	const bool bCulledOnConnectionCount = false;
+#endif
+
+	// Debug accumulators
+	FNativeClassAccumulator DormancyClassAccumulator;
+	FNativeClassAccumulator DistanceClassAccumulator;
+
+	// -------------------------------------------------------
+	// For Each Connection
+	// -------------------------------------------------------
+	
+	FGatheredReplicationActorLists GatheredReplicationListsForConnection;
+
+	for (UNetReplicationGraphConnection* ConnectionManager: Connections)
+	{
+		if (ConnectionManager->PrepareForReplication() == false)
+		{
+			// Connection is not ready to replicate
+			continue;
+		}
+
+		UNetConnection* const NetConnection = ConnectionManager->NetConnection;
+		APlayerController* const PC = NetConnection->PlayerController;
+		FPerConnectionActorInfoMap& ConnectionActorInfoMap = ConnectionManager->ActorInfoMap;
+		
+		repCheckf(NetConnection->GetReplicationConnectionDriver() == ConnectionManager, TEXT("NetConnection %s mismatch rep driver. %s vs %s"), *GetNameSafe(NetConnection), *GetNameSafe(NetConnection->GetReplicationConnectionDriver()), *GetNameSafe(ConnectionManager));
+		
+		// Send ClientAdjustments (movement RPCs) do this first and never let bandwidth saturation suppress these.
+		if (PC)
+		{
+			PC->SendClientAdjustment();
+		}
+
+		FBitWriter& ConnectionSendBuffer = NetConnection->SendBuffer;
+
+		const bool bEnableFullActorPrioritzationDetails = DO_REPGRAPH_DETAILS(bEnableFullActorPrioritizationDetailsAllConnections || ConnectionManager->bEnableFullActorPrioritizationDetails);
+
+		// ------------------------------------------------------------------------------------------
+		// FIXME: Packet budget is probably not necessary at this time. This should be simplified/removed.
+		// 
+		// Setup Packet Budget. This will determine:
+		//	1) What actor lists we will pull from the graph
+		//  2) How much data we allow each list to write to the send buffer before bailing
+		// ------------------------------------------------------------------------------------------
+
+		// Temp: legacy/dynamic budget based on queuedBits
+		static FPacketBudget LegacyBudget;
+		LegacyBudget.BudgetItems.Reset();
+
+		if (CVar_RepGraph_FixedBudget > 0.f)
+		{
+			LegacyBudget.AddBudgetForCategory( MAX_uint32, -CVar_RepGraph_FixedBudget );
+		}
+		else
+		{
+			LegacyBudget.AddBudgetForCategory( MAX_uint32, -NetConnection->QueuedBits );
+		}
+		
+		LegacyBudget.Finalize();
+
+		// PacketBudget defines our target sizes for various categories of replication lists
+		FPacketBudget& PacketBudget =  CVar_RepGraph_UseLegacyBudget ? LegacyBudget : GetPacketBudgetForConnection(ConnectionManager, FrameNum);
+		repCheckf(PacketBudget.BudgetItems.Num() > 0, TEXT("PacketBudget has no entries, cannot replicate actors without a budget."));
+
+		// This will track our packet budget as we fill up the packet
+		FPacketBudgetTracker PacketBudgetTracker(PacketBudget);
+
+		// Setup output for packet budget recording (this would be for logging/profiling).
+		if (FPacketBudgetRecordBuffer* PacketBudgetRecord = ConnectionManager->PacketRecordBuffer)
+		{
+			PacketBudgetTracker.SetOutputRecord(PacketBudgetRecord->GetNextWriteRecord(PacketBudget));
+		}
+
+		// --------------------------------------------------------------------------------------------------------------
+		// GATHER list of ReplicationLists for this connection
+		// --------------------------------------------------------------------------------------------------------------
+
+		// Determine Net Viewer/whatever for this connection (Maybe temp: should this be replaced with something more direct?
+		FNetViewer Viewer(NetConnection, 0.f);
+		const FVector ConnectionViewLocation = Viewer.ViewLocation;
+
+		GatheredReplicationListsForConnection.Reset();
+
+		const FConnectionGatherActorListParameters Parameters(Viewer, *ConnectionManager, NetConnection->ClientVisibleLevelNames, PacketBudget, FrameNum, GatheredReplicationListsForConnection);
+
+		int32 NumGatheredListsOnConnection = 0;
+		int32 NumGatheredActorsOnConnection = 0;
+		int32 NumPrioritizedActorsOnConnection = 0;
+
+		{
+			RG_QUICK_SCOPE_CYCLE_COUNTER(NET_ReplicateActors_GatherForConnection);
+
+			// Push default list category. Any nodes that don't specify a category will fallback to this
+			Parameters.AttemptPushListCategory(DefaultReplicationListCategory);
+
+			for (UReplicationGraphNode* Node : GlobalGraphNodes)
+			{
+				Node->GatherActorListsForConnection(Parameters);
+			}
+
+			for (UReplicationGraphNode* Node : ConnectionManager->ConnectionGraphNodes)
+			{
+				Node->GatherActorListsForConnection(Parameters);
+			}
+
+			if (GatheredReplicationListsForConnection.Num() == 0)
+			{
+				// No lists were returned, kind of weird but not fatal. Early out because code below assumes at least 1 list
+				UE_LOG(LogReplicationGraph, Warning, TEXT("No Replication Lists were returned for connection"));
+				return 0;
+			}
+		}
+
+		// --------------------------------------------------------------------------------------------------------------
+		// PRIORITIZE Gathered Actors For Connection
+		// --------------------------------------------------------------------------------------------------------------
+		{
+			RG_QUICK_SCOPE_CYCLE_COUNTER(NET_ReplicateActors_PrioritizeForConnection);
+
+			// We will make a prioritized list for each item in the packet budget. (Each item may accept multiple categories. Each list has one category)
+			// This means, depending on the packet budget, a gathered list could end up in multiple prioritized lists. This would not be desirable in most cases but is not explicitly forbidden.
+
+			PrioritizedReplicationLists.Reset();
+			for (const FPacketBudget::FItem& BudgetItem : PacketBudget.BudgetItems)
+			{
+				const FReplicationListCategoryMask BudgetCategory = BudgetItem.ListCategory;
+				FPrioritizedRepList* PrioritizedList = new (PrioritizedReplicationLists) FPrioritizedRepList(BudgetCategory);
+				
+				TArray<FPrioritizedRepList::FItem>* SortingArray = &PrioritizedList->Items; // FIXME: Items is a TArray and this will cause memory allocations
+
+				NumGatheredListsOnConnection += GatheredReplicationListsForConnection.Num();
+
+				for (FActorRepListRawView& List : GatheredReplicationListsForConnection)
+				{
+					if ((List.Category & BudgetCategory) > 0)
+					{
+						// Add actors from gathered list
+						NumGatheredActorsOnConnection += List.Num();
+						for (AActor* Actor : List)
+						{
+							RG_QUICK_SCOPE_CYCLE_COUNTER(Prioritize_InnerLoop);
+
+							// -----------------------------------------------------------------------------------------------------------------
+							//	Prioritize Actor for Connection: this is the main block of code for calculating a final score for this actor
+							//		-This is still pretty rough. It would be nice if this was customizable per project without suffering virtual calls.
+							// -----------------------------------------------------------------------------------------------------------------
+
+							FConnectionReplicationActorInfo& ConnectionData = ConnectionActorInfoMap.FindOrAdd(Actor);
+
+							RG_QUICK_SCOPE_CYCLE_COUNTER(Prioritize_InnerLoop_ConnGlobalLookUp);
+
+							// Skip if dormant on this connection. We want this to always be the first/quickest check.
+							if (ConnectionData.bDormantOnConnection)
+							{
+								DO_REPGRAPH_DETAILS(PrioritizedList->GetNextSkippedDebugDetails(Actor)->bWasDormant = true);
+								if (bCulledOnConnectionCount)
+								{
+									DormancyClassAccumulator.Increment(Actor->GetClass());
+								}
+								continue;
+							}
+
+							FGlobalActorReplicationInfo& GlobalData = GlobalActorReplicationInfoMap.Get(Actor);
+
+							RG_QUICK_SCOPE_CYCLE_COUNTER(Prioritize_InnerLoop_PostGlobalLookUp);
+
+							// Skip if its not time to replicate on this connection yet. We have to look at ForceNetUpdateFrame here. It would be possible to clear
+							// NextReplicationFrameNum on all connections when ForceNetUpdate is called. This probably means more work overall per frame though. Something to consider.
+							if (!ReadyForNextReplication(ConnectionData, GlobalData, FrameNum))
+							{
+								DO_REPGRAPH_DETAILS(PrioritizedList->GetNextSkippedDebugDetails(Actor)->FramesTillNextReplication = (FrameNum - ConnectionData.LastRepFrameNum));
+								continue;
+							}
+
+							RG_QUICK_SCOPE_CYCLE_COUNTER(Prioritize_InnerLoop_PostReady);
+							
+							// Output record for full debugging. This is not used in the actual sorting/prioritization of the list, just for logging/debugging purposes
+							FPrioritizedActorFullDebugDetails* DebugDetails = nullptr;
+							if (DO_REPGRAPH_DETAILS(UNLIKELY(bEnableFullActorPrioritzationDetails)))
+							{
+								DO_REPGRAPH_DETAILS(DebugDetails = PrioritizedList->GetNextFullDebugDetails(Actor));
+							}
+
+							float AccumulatedPriority = 0.f;
+
+							// -------------------
+							// Distance Scaling
+							// -------------------
+							if (GlobalData.Settings.DistancePriorityScale > 0.f)
+							{
+								const float MaxThreshold = 3000.f * 3000.f;
+								const float DistSq = (GlobalData.WorldLocation - ConnectionViewLocation).SizeSquared();
+
+								if (bDistanceCull && ConnectionData.CullDistanceSquared > 0.f && DistSq > ConnectionData.CullDistanceSquared)
+								{
+									DO_REPGRAPH_DETAILS(PrioritizedList->GetNextSkippedDebugDetails(Actor)->DistanceCulled = FMath::Sqrt(DistSq));
+									if (bCulledOnConnectionCount)
+									{
+										DistanceClassAccumulator.Increment(Actor->GetClass());
+									}
+									continue;
+								}
+
+								const float DistanceFactor = FMath::Clamp<float>( (DistSq) / MaxThreshold, 0.f, 1.f) * GlobalData.Settings.DistancePriorityScale;
+								AccumulatedPriority += DistanceFactor;
+								
+								if (DO_REPGRAPH_DETAILS(UNLIKELY(DebugDetails)))
+								{
+									DebugDetails->DistanceSq = DistSq;
+									DebugDetails->DistanceFactor = DistanceFactor;
+								}
+							}
+
+							RG_QUICK_SCOPE_CYCLE_COUNTER(Prioritize_InnerLoop_PostCull);
+
+							// Update the timeout frame number here. (Since this was returned by the graph, regardless if we end up replicating or not, we bump up the timeout frame num. This has to be done here because Distance Scaling can cull the actor
+							UpdateActorChannelCloseFrameNum(ConnectionData, GlobalData, FrameNum);
+
+							//UE_CLOG(DebugConnection, LogReplicationGraph, Display, TEXT("0x%X0x%X ConnectionData.ActorChannelCloseFrameNum=%d on %d"), Actor, NetConnection, ConnectionData.ActorChannelCloseFrameNum, FrameNum);
+
+							// -------------------
+							// Starvation Scaling
+							// -------------------
+							if (GlobalData.Settings.StarvationPriorityScale > 0.f)
+							{
+								// FIXME: Constant/Magic Num
+								const uint32 MaxFramesSinceLastRep = 20;
+								const uint32 FramesSinceLastRep = (FrameNum - ConnectionData.LastRepFrameNum);
+
+								const float StarvationFactor = 1.f - FMath::Clamp<float>((float)FramesSinceLastRep / (float)MaxFramesSinceLastRep, 0.f, 1.f);
+
+								AccumulatedPriority += StarvationFactor;
+
+								if (DO_REPGRAPH_DETAILS(UNLIKELY(DebugDetails)))
+								{
+									DebugDetails->FramesSinceLastRap = FramesSinceLastRep;
+									DebugDetails->StarvationFactor = StarvationFactor;
+								}
+							}
+
+							// -------------------
+							//	Game code priority
+							// -------------------
+							
+							if (GlobalData.ForceNetUpdateFrame > 0)
+							{
+								uint32 ForceNetUpdateDelta = GlobalData.ForceNetUpdateFrame - ConnectionData.LastRepFrameNum;
+								if ( ForceNetUpdateDelta > 0 )
+								{
+									// Note that in legacy ForceNetUpdate did not actually bump priority. This gives us a hard coded bump if we haven't replicated since the last ForceNetUpdate frame.
+									AccumulatedPriority += 1.f;
+
+									if (DO_REPGRAPH_DETAILS(UNLIKELY(DebugDetails)))
+									{
+										DebugDetails->GameCodeScaling = 1.f;
+									}
+								}
+							}
+							
+							SortingArray->Emplace(FPrioritizedRepList::FItem(AccumulatedPriority, Actor /*, &GlobalData, &ConnectionData*/ ));
+						}
+					}
+				}
+
+				{
+					// Sort the merged priority list. We could potentially move this into the replicate loop below, this could potentially save use from sorting arrays that don't fit into the budget
+					RG_QUICK_SCOPE_CYCLE_COUNTER(NET_ReplicateActors_PrioritizeForConnection_Sort);
+					NumPrioritizedActorsOnConnection += SortingArray->Num();
+					SortingArray->Sort();
+				}
+			}
+		}
+
+		// --------------------------------------------------------------------------------------------------------------
+		// REPLICATE Actors For Connection
+		// --------------------------------------------------------------------------------------------------------------
+		{
+			RG_QUICK_SCOPE_CYCLE_COUNTER(NET_ReplicateActors_ReplicateActorsForConnection);
+
+			int32 ListIdx=0;
+			for (; ListIdx < PrioritizedReplicationLists.Num(); ++ListIdx)
+			{
+				const FPrioritizedRepList& List = PrioritizedReplicationLists[ListIdx];
+				
+				// Tell packet budget tracker where these bits should be tracked
+				PacketBudgetTracker.SetBudgetItemIndex(ListIdx);
+				
+				for (int32 ActorIdx=0; ActorIdx < List.Items.Num(); ++ActorIdx)
+				{
+					const FPrioritizedRepList::FItem& RepItem = List.Items[ActorIdx];
+
+					AActor* Actor = RepItem.Actor;
+					FConnectionReplicationActorInfo& ActorInfo = ConnectionActorInfoMap.FindOrAdd(Actor); // FIXME: !!! *RepItem.ConnectionData;
+
+					// Always skip if we've already replicated this frame. This happens if an actor is in more than one replication list
+					if (ActorInfo.LastRepFrameNum == FrameNum)
+					{
+						INC_DWORD_STAT_BY( STAT_NetRepActorListDupes, 1 );
+						continue;
+					}
+
+					FGlobalActorReplicationInfo& GlobalActorInfo = GlobalActorReplicationInfoMap.Get(Actor);// FIXME *RepItem.GlobalData;
+
+					int64 BitsWritten = ReplicateSingleActor(Actor, ActorInfo, GlobalActorInfo, ConnectionActorInfoMap, NetConnection, FrameNum);
+
+					// --------------------------------------------------
+					//	Update Packet Budget Tracking
+					// --------------------------------------------------
+					
+					if (PacketBudgetTracker.TrackSentBits(BitsWritten) == false || IsConnectionReady(NetConnection) == false)
+					{
+						// We've exceeded the budget for this category of replication list.
+						RG_QUICK_SCOPE_CYCLE_COUNTER(NET_ReplicateActors_PartialStarvedActorList);
+						HandleStarvedActorList(List, ActorIdx+1, ConnectionActorInfoMap, FrameNum);
+						GNumSaturatedConnections++;
+						break;
+					}
+				}
+			}
+
+			// ------------------------------------------
+			// Handle starved lists that we didn't get to
+			// ------------------------------------------
+			{
+				
+				RG_QUICK_SCOPE_CYCLE_COUNTER(NET_ReplicateActors_StarvedActorLists);
+				for (; ListIdx < PrioritizedReplicationLists.Num(); ++ListIdx)
+				{
+					const FPrioritizedRepList& List = PrioritizedReplicationLists[ListIdx];
+					HandleStarvedActorList(List, 0, ConnectionActorInfoMap, FrameNum);
+				}
+			}
+
+			// ------------------------------------------
+			// Handle stale, no longer relevant, actor channels.
+			// ------------------------------------------			
+			{
+				// TODO: Use a BitArray to track which channels didn't get repped or starved this frame ?
+				RG_QUICK_SCOPE_CYCLE_COUNTER(NET_ReplicateActors_LookForNonRelevantChannels);
+
+				for (auto MapIt = ConnectionActorInfoMap.CreateIterator(); MapIt; ++MapIt)
+				{
+					FConnectionReplicationActorInfo& ConnectionActorInfo = MapIt.Value();
+					if (ConnectionActorInfo.Channel != nullptr && ConnectionActorInfo.ActorChannelCloseFrameNum > 0 && ConnectionActorInfo.ActorChannelCloseFrameNum <= FrameNum) //ConnectionActorInfo.StarvedFrameNum == 0 && ConnectionActorInfo.LastRepFrameNum < TimeOutFrame && )
+					{
+						AActor* Actor = MapIt.Key();
+						if (Actor->IsNetStartupActor())
+						{
+							continue;
+						}
+						
+						//UE_CLOG(DebugConnection, LogReplicationGraph, Display, TEXT("Closing Actor Channel:0x%x 0x%X0x%X, %s %d <= %d"), ConnectionActorInfo.Channel, Actor, NetConnection, *GetNameSafe(ConnectionActorInfo.Channel->Actor), ConnectionActorInfo.ActorChannelCloseFrameNum, FrameNum);
+
+						INC_DWORD_STAT_BY( STAT_NetActorChannelsClosed, 1 );
+						PacketBudgetTracker.TrackActorChannelCloseBits( ConnectionActorInfo.Channel->Close() );
+					}
+				}
+			}
+
+			// ------------------------------------------
+			// Handle Destruction Infos. These are actors that have been destroyed on the server but that we need to tell the client about.
+			// ------------------------------------------
+			{
+				RG_QUICK_SCOPE_CYCLE_COUNTER(NET_ReplicateActors_ReplicateDestructionInfos);
+				PacketBudgetTracker.TrackDestructionInfoBits( ConnectionManager->ReplicateDestructionInfos(ConnectionViewLocation, DestructInfoMaxDistanceSquared) );
+			}
+
+#if !(UE_BUILD_SHIPPING || UE_BUILD_TEST)
+			{
+				RG_QUICK_SCOPE_CYCLE_COUNTER(NET_ReplicateActors_ReplicateDebugActor);
+				if (ConnectionManager->DebugActor)
+				{
+					FGlobalActorReplicationInfo& GlobalInfo = GlobalActorReplicationInfoMap.Get(ConnectionManager->DebugActor);
+					FConnectionReplicationActorInfo& ActorInfo = ConnectionActorInfoMap.FindOrAdd(ConnectionManager->DebugActor);
+					ReplicateSingleActor(ConnectionManager->DebugActor, ActorInfo, GlobalInfo, ConnectionActorInfoMap, NetConnection, FrameNum);
+				}
+			}
+#endif
+
+			// Broadcast the list we just handled. This is intended to be for debugging/logging features.
+			ConnectionManager->OnPostReplicatePrioritizeLists.Broadcast( ConnectionManager, TArrayView<FPrioritizedRepList>(PrioritizedReplicationLists.GetData(), PrioritizedReplicationLists.Num()) );
+
+			if (bCulledOnConnectionCount)
+			{
+				// Reset the CVar this only counts for one frame
+				CVar_RepGraph_PrintCulledOnConnectionClasses = 0;
+				UE_LOG(LogReplicationGraph, Display, TEXT("Dormant Culled classes: %s"), *DormancyClassAccumulator.BuildString());
+				UE_LOG(LogReplicationGraph, Display, TEXT("Dist Culled classes: %s"), *DistanceClassAccumulator.BuildString());
+				UE_LOG(LogReplicationGraph, Display, TEXT("Saturated Connections: %d"), GNumSaturatedConnections);
+				UE_LOG(LogReplicationGraph, Display, TEXT(""));
+
+				UE_LOG(LogReplicationGraph, Display, TEXT("Gathered Lists: %d Gathered Actors: %d  PrioritizedActors: %d"), NumGatheredListsOnConnection, NumGatheredActorsOnConnection, NumPrioritizedActorsOnConnection);
+				UE_LOG(LogReplicationGraph, Display, TEXT("Connection Loaded Streaming Levels: %d"), Parameters.ClientVisibleLevelNamesRef.Num());
+			}
+		}
+	}
+	
+	SET_DWORD_STAT(STAT_NumProcessedConnections, Connections.Num());
+
+	if (CVar_RepGraph_PrintTrackClassReplication)
+	{
+		CVar_RepGraph_PrintTrackClassReplication = 0;
+		UE_LOG(LogReplicationGraph, Display, TEXT("Changed Classes: %s"), *ChangeClassAccumulator.BuildString());
+		UE_LOG(LogReplicationGraph, Display, TEXT("No Change Classes: %s"), *NoChangeClassAccumulator.BuildString());
+	}
+
+	return 0;
+}
+
+int64 UReplicationGraph::ReplicateSingleActor(AActor* Actor, FConnectionReplicationActorInfo& ActorInfo, FGlobalActorReplicationInfo& GlobalActorInfo, FPerConnectionActorInfoMap& ConnectionActorInfoMap, UNetConnection* NetConnection, const uint32 FrameNum)
+{
+	RG_QUICK_SCOPE_CYCLE_COUNTER(NET_ReplicateActors_ReplicateSingleActor);
+
+	if (ActorInfo.Channel && ActorInfo.Channel->Closing)
+	{
+		// We are waiting for the client to ack this actor channel's close bunch.
+		return 0;
+	}
+
+	ActorInfo.LastRepFrameNum = FrameNum;
+	ActorInfo.StarvedFrameNum = 0;
+	ActorInfo.NextReplicationFrameNum = FrameNum + ActorInfo.ReplicationPeriodFrame;
+
+	// -------------------------------------------
+	// FIXME: Call PreReplication if necessary. It kind of sucks to have an additional TMap lookup inside this loop, but hard to know otherwise what actors need it called.
+	// Lots of crappy code in ::PreReplication, but its probably unavoidable right now. This is still probably better than legacy calling it on every actor every tick that might replicate.
+	// -------------------------------------------
+	if (GlobalActorInfo.LastPreReplicationFrame != FrameNum)
+	{
+		RG_QUICK_SCOPE_CYCLE_COUNTER(NET_ReplicateActors_CallPreReplication);
+		GlobalActorInfo.LastPreReplicationFrame = FrameNum;
+
+		Actor->CallPreReplication(NetDriver);
+	}
+
+	const bool bWantsToGoDormant = GlobalActorInfo.bWantsToBeDormant;
+	TActorRepListViewBase<FActorRepList*> DependentActorList(GlobalActorInfo.DependentActorList.RepList.GetReference());
+
+	if (ActorInfo.Channel == nullptr)
+	{
+		// Create a new channel for this actor.
+		INC_DWORD_STAT_BY( STAT_NetActorChannelsOpened, 1 );
+		ActorInfo.Channel = (UActorChannel*)NetConnection->CreateChannel( CHTYPE_Actor, 1 );
+		if ( !ActorInfo.Channel )
+		{
+			return 0;
+		}
+
+		//UE_LOG(LogReplicationGraph, Display, TEXT("Created Actor Channel:0x%x 0x%X0x%X, %d"), ActorInfo.Channel, Actor, NetConnection, FrameNum);
+					
+		// This will unfortunately cause a callback to this  UNetReplicationGraphConnection and will relook up the ActorInfoMap and set the channel that we already have set.
+		// This is currently unavoidable because channels are created from different code paths (some outside of this loop).
+		ActorInfo.Channel->SetChannelActor( Actor );
+	}
+
+	if (UNLIKELY(bWantsToGoDormant))
+	{
+		ActorInfo.Channel->StartBecomingDormant();
+	}
+
+	int64 BitsWritten = 0;
+					
+	if (UNLIKELY(ActorInfo.bTearOff))
+	{
+		// Replicate and immediately close in tear off case
+		BitsWritten = ActorInfo.Channel->ReplicateActor();
+		BitsWritten += ActorInfo.Channel->Close();
+	}
+	else
+	{
+		// Just replicate normally
+		BitsWritten = ActorInfo.Channel->ReplicateActor();
+	}
+
+	if (bTrackClassReplication)
+	{
+		if (BitsWritten > 0)
+		{
+			ChangeClassAccumulator.Increment(Actor->GetClass());
+		}
+		else
+		{
+			NoChangeClassAccumulator.Increment(Actor->GetClass());
+		}
+	}
+	
+
+	// ----------------------------
+	//	Dependent actors
+	// ----------------------------
+	if (DependentActorList.IsValid())
+	{
+		RG_QUICK_SCOPE_CYCLE_COUNTER(NET_ReplicateActors_DependentActors);
+
+		const int32 CloseFrameNum = ActorInfo.ActorChannelCloseFrameNum; // cache this off since subsequent calls to FindOrAdd will invalidate ActorInfo.
+
+		for (AActor* DependentActor : DependentActorList)
+		{
+			repCheck(DependentActor);
+
+			FConnectionReplicationActorInfo& DependentActorConnectionInfo = ConnectionActorInfoMap.FindOrAdd(DependentActor);
+			FGlobalActorReplicationInfo& DependentActorGlobalData = GlobalActorReplicationInfoMap.Get(DependentActor);
+
+			// Dependent actor channel will stay open as long as the owning actor channel is open
+			DependentActorConnectionInfo.ActorChannelCloseFrameNum = FMath::Max<uint32>(CloseFrameNum, DependentActorConnectionInfo.ActorChannelCloseFrameNum);
+
+			if (!ReadyForNextReplication(DependentActorConnectionInfo, DependentActorGlobalData, FrameNum))
+			{
+				continue;
+			}
+
+			//UE_LOG(LogReplicationGraph, Display, TEXT("DependentActor %s %s. NextReplicationFrameNum: %d. FrameNum: %d. ForceNetUpdateFrame: %d. LastRepFrameNum: %d."), *DependentActor->GetPathName(), *NetConnection->GetName(), DependentActorConnectionInfo.NextReplicationFrameNum, FrameNum, DependentActorGlobalData.ForceNetUpdateFrame, DependentActorConnectionInfo.LastRepFrameNum);
+			BitsWritten += ReplicateSingleActor(DependentActor, DependentActorConnectionInfo, DependentActorGlobalData, ConnectionActorInfoMap, NetConnection, FrameNum);
+		}					
+	}
+
+	return BitsWritten;
+}
+
+void UReplicationGraph::HandleStarvedActorList(const FPrioritizedRepList& List, int32 StartIdx, FPerConnectionActorInfoMap& ConnectionActorInfoMap, uint32 FrameNum)
+{
+	for (int32 ActorIdx=StartIdx; ActorIdx < List.Items.Num(); ++ActorIdx)
+	{
+		// FIXME[: we are looking up ConnectionActorInfoMap even though this could be stored in FPrioritizedRepList::FItem if the pointers were stable.
+		const FPrioritizedRepList::FItem& RepItem = List.Items[ActorIdx];
+		FConnectionReplicationActorInfo& ActorInfo = ConnectionActorInfoMap.FindOrAdd(RepItem.Actor); // FIXME: !!! *RepItem.ConnectionData;
+		
+		// Only update starve frame if not already starved (we want to use this to measure "how long have you been starved for")
+		if (ActorInfo.StarvedFrameNum == 0)
+		{
+			ActorInfo.StarvedFrameNum = FrameNum;
+		}
+
+		// Update dependent actor's timeout frame
+		FGlobalActorReplicationInfo& GlobalActorInfo = GlobalActorReplicationInfoMap.Get(RepItem.Actor);
+		TActorRepListViewBase<FActorRepList*> DependentActorList(GlobalActorInfo.DependentActorList.RepList.GetReference());
+		
+		if (DependentActorList.IsValid())
+		{
+			const uint32 CloseFrameNum = ActorInfo.ActorChannelCloseFrameNum; // cache this off since subsequent calls to FindOrAdd will invalidate ActorInfo.
+			for (AActor* DependentActor : DependentActorList)
+			{
+				FConnectionReplicationActorInfo& DependentActorConnectionInfo = ConnectionActorInfoMap.FindOrAdd(DependentActor);
+				DependentActorConnectionInfo.ActorChannelCloseFrameNum = FMath::Max<uint32>(CloseFrameNum, DependentActorConnectionInfo.ActorChannelCloseFrameNum);
+			}
+		}
+	}
+}
+
+void UReplicationGraph::UpdateActorChannelCloseFrameNum(FConnectionReplicationActorInfo& ConnectionData, const FGlobalActorReplicationInfo& GlobalData, const uint32 FrameNum) const
+{
+	// Only update if the actor has a timeout set
+	if (GlobalData.Settings.ActorChannelFrameTimeout > 0)
+	{
+		const uint32 NewCloseFrameNum = FrameNum + ConnectionData.ReplicationPeriodFrame + GlobalData.Settings.ActorChannelFrameTimeout + GlobalActorChannelFrameNumTimeout;
+		ConnectionData.ActorChannelCloseFrameNum = FMath::Max<uint32>(ConnectionData.ActorChannelCloseFrameNum, NewCloseFrameNum); // Never go backwards, something else could have bumped it up further intentionally
+	}
+}
+
+bool UReplicationGraph::ProcessRemoteFunction(class AActor* Actor, UFunction* Function, void* Parameters, FOutParmRec* OutParms, FFrame* Stack, class UObject* SubObject )
+{
+	// ----------------------------------
+	// Setup
+	// ----------------------------------
+
+	if (IsActorValidForReplication(Actor) == false || Actor->IsActorBeingDestroyed())
+	{
+		return true;
+	}
+
+	// get the top most function
+	while( Function->GetSuperFunction() )
+	{
+		Function = Function->GetSuperFunction();
+	}
+
+	// If we have a subobject, thats who we are actually calling this on. If no subobject, we are calling on the actor.
+	UObject* TargetObj = SubObject ? SubObject : Actor;
+
+	// Make sure this function exists for both parties.
+	const FClassNetCache* ClassCache = NetDriver->NetCache->GetClassNetCache( TargetObj->GetClass() );
+	if (!ClassCache)
+	{
+		UE_LOG(LogReplicationGraph, Warning, TEXT("ClassNetCache empty, not calling %s::%s"), *Actor->GetName(), *Function->GetName());
+		return true;
+	}
+		
+	const FFieldNetCache* FieldCache = ClassCache->GetFromField( Function );
+	if ( !FieldCache )
+	{
+		UE_LOG(LogReplicationGraph, Warning, TEXT("FieldCache empty, not calling %s::%s"), *Actor->GetName(), *Function->GetName());
+		return true;
+	}
+
+	// ----------------------------------
+	// Multicast
+	// ----------------------------------
+
+	if ((Function->FunctionFlags & FUNC_NetMulticast))
+	{
+		TSharedPtr<FRepLayout> RepLayout = NetDriver->GetFunctionRepLayout( Function );
+
+		TOptional<FVector> ActorLocation;
+
+		UNetDriver::ERemoteFunctionSendPolicy SendPolicy = UNetDriver::Default;
+		if (CVar_RepGraph_EnableRPCSendPolicy > 0)
+		{
+			if (FRPCSendPolicyInfo* FuncSendPolicy = RPCSendPolicyMap.Find(FObjectKey(Function)))
+			{
+				if (FuncSendPolicy->bSendImmediately)
+				{
+					SendPolicy = UNetDriver::ForceSend;
+				}
+			}
+		}
+
+		RepLayout->BuildSharedSerializationForRPC(Parameters);
+		FGlobalActorReplicationInfo& GlobalInfo = GlobalActorReplicationInfoMap.Get(Actor);
+		const float CullDistanceSquared = GlobalInfo.Settings.CullDistanceSquared;
+		
+		for (UNetReplicationGraphConnection* Manager : Connections)
+		{
+			FConnectionReplicationActorInfo& ConnectionActorInfo = Manager->ActorInfoMap.FindOrAdd(Actor);
+			UNetConnection* NetConnection = Manager->NetConnection;
+
+			// This connection isn't ready yet
+			if (NetConnection->ViewTarget == nullptr)
+			{
+				continue;
+			}
+			
+			if (ConnectionActorInfo.Channel == nullptr)
+			{
+				// There is no actor channel here. Ideally we would just ignore this but in the case of net dormancy, this may be an actor that will replicate on the next frame.
+				// If the actor is dormant and is a distance culled actor, we can probably safely assume this connection will open a channel for the actor on the next rep frame.
+				// This isn't perfect and we may want a per-function or per-actor policy that allows to dictate what happens in this situation.
+
+				// Actors being destroyed (Building hit with rocket) will wake up before this gets hit. So dormancy really cant be relied on here.
+				// if (Actor->NetDormancy > DORM_Awake)
+				{
+					if (CullDistanceSquared > 0)
+					{
+						FNetViewer Viewer(NetConnection, 0.f);
+						if (ActorLocation.IsSet() == false)
+						{
+							ActorLocation = Actor->GetActorLocation();
+						}
+
+						const float DistSq = (ActorLocation.GetValue() - Viewer.ViewLocation).SizeSquared();
+						if (DistSq <= CullDistanceSquared)
+						{
+							// We are within range, we will open a channel now for this actor and call the RPC on it
+							ConnectionActorInfo.Channel = (UActorChannel *)NetConnection->CreateChannel( CHTYPE_Actor, 1 );
+							ConnectionActorInfo.Channel->SetChannelActor(Actor);
+
+							// Update timeout frame name. We would run into problems if we open the channel, queue a bunch, and then it timeouts before RepGraph replicates properties.
+							UpdateActorChannelCloseFrameNum(ConnectionActorInfo, GlobalInfo, ReplicationGraphFrame+1 /** Plus one to error on safe side. RepFrame num will be incremented in the next tick */ );
+						}
+					}
+				}
+			}
+			
+			if (ConnectionActorInfo.Channel)
+			{
+				NetDriver->ProcessRemoteFunctionForChannel(ConnectionActorInfo.Channel, ClassCache, FieldCache, TargetObj, NetConnection, Function, Parameters, OutParms, Stack, true, SendPolicy);
+
+				if (SendPolicy == UNetDriver::ForceSend)
+				{
+					RG_QUICK_SCOPE_CYCLE_COUNTER(RPC_FORCE_FLUSH_NET);
+					NetConnection->FlushNet();
+				}
+
+			}
+		}
+
+		RepLayout->ClearSharedSerializationForRPC();
+		return true;
+	}
+
+	// ----------------------------------
+	// Single Connection
+	// ----------------------------------
+	
+	UNetConnection* Connection = Actor->GetNetConnection();
+	if (Connection)
+	{
+		if (((Function->FunctionFlags & FUNC_NetReliable) == 0) && !IsConnectionReady(Connection))
+		{
+			return true;
+		}
+	
+		if (Connection->State == USOCK_Closed)
+		{
+			return true;
+		}
+
+		UActorChannel* Ch = Connection->FindActorChannelRef(Actor);
+		if (Ch == nullptr)
+		{
+			if ( Actor->IsPendingKillPending() || !NetDriver->IsLevelInitializedForActor(Actor, Connection) )
+			{
+				// We can't open a channel for this actor here
+				return true;
+			}
+
+			Ch = (UActorChannel *)Connection->CreateChannel( CHTYPE_Actor, 1 );
+			Ch->SetChannelActor(Actor);
+		}
+
+		NetDriver->ProcessRemoteFunctionForChannel(Ch, ClassCache, FieldCache, TargetObj, Connection, Function, Parameters, OutParms, Stack, true);
+	}
+	else
+	{
+		UE_LOG(LogNet, Warning, TEXT("UReplicationGraph::ProcessRemoteFunction: No owning connection for actor %s. Function %s will not be processed."), *Actor->GetName(), *Function->GetName());
+	}
+
+	// return true because we don't want the net driver to do anything else
+	return true;
+}
+
+bool UReplicationGraph::IsConnectionReady(UNetConnection* Connection)
+{
+	return Connection->QueuedBits + Connection->SendBuffer.GetNumBits() <= 0;
+}
+
+// --------------------------------------------------------------------------------------------------------------------------------------------
+// --------------------------------------------------------------------------------------------------------------------------------------------
+// --------------------------------------------------------------------------------------------------------------------------------------------
+
+
+UNetReplicationGraphConnection::UNetReplicationGraphConnection()
+{
+
+}
+void UNetReplicationGraphConnection::NotifyActorChannelAdded(AActor* Actor, class UActorChannel* Channel)
+{
+	FConnectionReplicationActorInfo& ActorInfo = ActorInfoMap.FindOrAdd(Actor);
+	ActorInfo.Channel = Channel;
+}
+
+void UNetReplicationGraphConnection::NotifyActorChannelRemoved(AActor* Actor)
+{
+	// No need to do anything here. This is called when an actor channel is closed, but
+	// we're still waiting for the close bunch to be acked. Until then, we can't safely replicate
+	// the actor from this channel. See NotifyActorChannelCleanedUp.
+}
+
+void UNetReplicationGraphConnection::NotifyActorChannelCleanedUp(UActorChannel* Channel)
+{
+	if (Channel)
+	{
+		// No existing way to quickly index from actor channel -> ActorInfo. May want a way to speed this up.
+		// The Actor pointer on the channel would have been set to null previously when the channel was closed,
+		// so we can't use that to look up the actor info by key.
+		// Also, the actor may be destroyed and garbage collected before this point.
+		for (auto It = ActorInfoMap.CreateIterator(); It; ++It)
+		{
+			if (It->Value.Channel == Channel)
+			{
+				if (Channel->Dormant)
+				{
+					// If the actor is just going dormant, clear the channel reference but leave the ActorInfo
+					// so that the graph can continue to track it.
+					It->Value.Channel = nullptr;
+				}
+				else
+				{
+					// If the channel wasn't cleaned up for dormancy, the graph doesn't need to track it anymore.
+					// Remove the ActorInfo and allow a new entry to be created and channel opened if the actor
+					// that was on this channel needs to replicate again.
+					It.RemoveCurrent();
+				}
+				break;
+			}
+		}
+	}
+}
+
+void UNetReplicationGraphConnection::InitForGraph(UReplicationGraph* Graph)
+{
+	// The per-connection data needs to know about the global data map so that it can pull defaults from it when we initialize a new actor
+	TSharedPtr<FReplicationGraphGlobalData> Globals = Graph ? Graph->GetGraphGlobals() : nullptr;
+	if (Globals.IsValid())
+	{
+		ActorInfoMap.SetGlobalMap(Globals->GlobalActorReplicationInfoMap);
+	}
+}
+
+void UNetReplicationGraphConnection::InitForConnection(UNetConnection* InConnection)
+{
+	NetConnection = InConnection;
+	InConnection->SetReplicationConnectionDriver(this);
+
+#if !(UE_BUILD_SHIPPING || UE_BUILD_TEST)
+	DebugActor = GetWorld()->SpawnActor<AReplicationGraphDebugActor>();
+	DebugActor->ConnectionManager = this;
+	DebugActor->ReplicationGraph = Cast<UReplicationGraph>(GetOuter());
+#endif
+
+#if 0
+	// This does not work because the control channel hasn't been opened yet. Could be moved further down the init path or in ServerReplicateActors.
+	FString TestStr(TEXT("Replication Graph is Enabled!"));	
+	FNetControlMessage<NMT_DebugText>::Send(InConnection,TestStr);
+	InConnection->FlushNet();
+#endif
+}
+
+void UNetReplicationGraphConnection::AddConnectionGraphNode(UReplicationGraphNode* Node)
+{
+	ConnectionGraphNodes.Add(Node);
+}
+
+void UNetReplicationGraphConnection::RemoveConnectionGraphNode(UReplicationGraphNode* Node)
+{
+	ConnectionGraphNodes.Remove(Node);
+}
+
+bool UNetReplicationGraphConnection::PrepareForReplication()
+{
+	NetConnection->ViewTarget = NetConnection->PlayerController ? NetConnection->PlayerController->GetViewTarget() : NetConnection->OwningActor;
+	return (NetConnection->ViewTarget != nullptr);
+}
+
+void UNetReplicationGraphConnection::NotifyAddDestructionInfo(FActorDestructionInfo* DestructInfo)
+{
+	if (DestructInfo->StreamingLevelName != NAME_None)
+	{
+		if (NetConnection->ClientVisibleLevelNames.Contains(DestructInfo->StreamingLevelName) == false)
+		{
+			// This client does not have this streaming level loaded. We should get notified again via UNetConnection::UpdateLevelVisibility
+			// (This should be enough. Legacy system would add the info and then do the level check in ::ServerReplicateActors, but this should be unnecessary)
+			return;
+		}
+	}
+
+	PendingDestructInfoList.Emplace( FCachedDestructInfo(DestructInfo) );
+}
+
+void UNetReplicationGraphConnection::NotifyRemoveDestructionInfo(FActorDestructionInfo* DestructInfo)
+{
+	int32 RemoveIdx = PendingDestructInfoList.IndexOfByKey(DestructInfo);
+	if (RemoveIdx != INDEX_NONE)
+	{
+		PendingDestructInfoList.RemoveAt(RemoveIdx, 1, false);
+	}
+}
+
+void UNetReplicationGraphConnection::NotifyResetDestructionInfo()
+{
+	PendingDestructInfoList.Reset();
+}
+
+void UNetReplicationGraphConnection::NotifyClientVisibleLevelNamesAdd(FName LevelName, UWorld* StreamingWorld) 
+{
+	// Undormant every actor in this world for this connection.
+	if (StreamingWorld && StreamingWorld->PersistentLevel)
+	{
+		TArray<AActor*>& Actors = StreamingWorld->PersistentLevel->Actors;
+		for (AActor* Actor : Actors)
+		{
+			if (Actor && Actor->GetIsReplicated() && (Actor->NetDormancy == DORM_DormantAll))
+			{
+				if (FConnectionReplicationActorInfo* ActorInfo = ActorInfoMap.Find(Actor))
+				{
+					ActorInfo->bDormantOnConnection = false;
+				}
+			}
+		}
+	}
+
+	OnClientVisibleLevelNameAdd.Broadcast(LevelName, StreamingWorld);
+
+}
+
+int64 UNetReplicationGraphConnection::ReplicateDestructionInfos(const FVector& ConnectionViewLocation, const float DestructInfoMaxDistanceSquared)
+{
+	const float X = ConnectionViewLocation.X;
+	const float Y = ConnectionViewLocation.Y;
+
+	int64 NumBits = 0;
+	for (int32 idx=PendingDestructInfoList.Num()-1; idx >=0; --idx)
+	{
+		const FCachedDestructInfo& Info = PendingDestructInfoList[idx];
+		float DistSquared = FMath::Square(Info.CachedPosition.X-X) + FMath::Square(Info.CachedPosition.Y-Y);
+
+		if (DistSquared < DestructInfoMaxDistanceSquared)
+		{
+			UActorChannel* Channel = (UActorChannel*)NetConnection->CreateChannel( CHTYPE_Actor, 1 );
+			if ( Channel )
+			{
+				NumBits += Channel->SetChannelActorForDestroy( Info.DestructionInfo );
+			}
+
+			PendingDestructInfoList.RemoveAtSwap(idx, 1, false);
+		}
+	}
+
+	return NumBits;
+}
+
+// --------------------------------------------------------------------------------------------------------------------------------------------
+// --------------------------------------------------------------------------------------------------------------------------------------------
+// --------------------------------------------------------------------------------------------------------------------------------------------
+
+UReplicationGraphNode::FCreateChildNodeFunc UReplicationGraphNode::DefaultCreatChildNodeFunc = [](UReplicationGraphNode* Parent)
+{
+	return Parent->CreateChildNode<UReplicationGraphNode_ActorList>();
+};
+
+// --------------------------------------------------------------------------------------------------------------------------------------------
+
+UReplicationGraphNode::UReplicationGraphNode()
+{
+	// The default implementation is to create an actor list node for your children
+	CreateChildSceneNodeFunc = DefaultCreatChildNodeFunc;
+}
+
+void UReplicationGraphNode::NotifyResetAllNetworkActors()
+{
+	for (UReplicationGraphNode* ChildNode : AllChildNodes)
+	{
+		ChildNode->NotifyResetAllNetworkActors();
+	}
+}
+
+void UReplicationGraphNode::TearDown()
+{
+	for (UReplicationGraphNode* Node : AllChildNodes)
+	{
+		Node->TearDown();
+	}
+
+	MarkPendingKill();
+}
+
+// --------------------------------------------------------------------------------------------------------------------------------------------
+void FStreamingLevelActorListCollection::AddActor(const FNewReplicatedActorInfo& ActorInfo)
+{
+	FStreamingLevelActors* Item = StreamingLevelLists.FindByKey(ActorInfo.StreamingLevelName);
+	if (!Item)
+	{
+		Item = new (StreamingLevelLists) FStreamingLevelActors(ActorInfo.StreamingLevelName);
+	}
+
+	if (CVar_RepGraph_Verify)
+	{
+		ensureMsgf(Item->ReplicationActorList.Contains(ActorInfo.Actor) == false, TEXT("%s being added to %s twice! Streaming level: %s"), *GetActorRepListTypeDebugString(ActorInfo.Actor), *ActorInfo.StreamingLevelName.ToString() );
+	}
+	Item->ReplicationActorList.Add(ActorInfo.Actor);
+
+}
+
+void FStreamingLevelActorListCollection::RemoveActor(const FNewReplicatedActorInfo& ActorInfo, bool bWarnIfNotFound, UReplicationGraphNode* Outer)
+{
+	for (FStreamingLevelActors& StreamingList : StreamingLevelLists)
+	{
+		if (StreamingList.StreamingLevelName == ActorInfo.StreamingLevelName)
+		{
+			if (!StreamingList.ReplicationActorList.Remove(ActorInfo.Actor) && bWarnIfNotFound)
+			{
+				UE_LOG(LogReplicationGraph, Warning, TEXT("Attempted to remove %s from list %s but it was not found. (StreamingLevelName == %s)"), *GetActorRepListTypeDebugString(ActorInfo.Actor), *GetPathNameSafe(Outer), *ActorInfo.StreamingLevelName.ToString() );
+			}
+
+			if (CVar_RepGraph_Verify)
+			{
+				ensureMsgf(StreamingList.ReplicationActorList.Contains(ActorInfo.Actor) == false, TEXT("Actor %s is still in %s after removal. Streaming Level: %s"), *GetActorRepListTypeDebugString(ActorInfo.Actor), *GetPathNameSafe(Outer));
+			}
+			break;
+		}
+	}
+}
+
+void FStreamingLevelActorListCollection::Reset()
+{
+	for (FStreamingLevelActors& StreamingList : StreamingLevelLists)
+	{
+		StreamingList.ReplicationActorList.Reset();
+	}
+}
+
+void FStreamingLevelActorListCollection::Gather(const FConnectionGatherActorListParameters& Params)
+{
+	for (const FStreamingLevelActors& StreamingList : StreamingLevelLists)
+	{
+		if (Params.CheckClientVisibilityForLevel(StreamingList.StreamingLevelName))
+		{
+			Params.OutGatheredReplicationLists.AddReplicationActorList(StreamingList.ReplicationActorList);
+		}
+		else
+		{
+			UE_LOG(LogReplicationGraph, Verbose, TEXT("Level Not Loaded %s. (Client has %d levels loaded)"), *StreamingList.StreamingLevelName.ToString(), Params.ClientVisibleLevelNamesRef.Num());
+		}
+	}
+}
+
+void FStreamingLevelActorListCollection::DeepCopyFrom(const FStreamingLevelActorListCollection& Source)
+{
+	StreamingLevelLists.Reset();
+	for (const FStreamingLevelActors& StreamingLevel : Source.StreamingLevelLists)
+	{
+		if (StreamingLevel.ReplicationActorList.Num() > 0)
+		{
+			FStreamingLevelActors* NewStreamingLevel = new (StreamingLevelLists)FStreamingLevelActors(StreamingLevel.StreamingLevelName);
+			NewStreamingLevel->ReplicationActorList.CopyContentsFrom(StreamingLevel.ReplicationActorList);
+			ensure(NewStreamingLevel->ReplicationActorList.Num() == StreamingLevel.ReplicationActorList.Num());
+		}
+	}
+}
+
+void FStreamingLevelActorListCollection::GetAll_Debug(TArray<FActorRepListType>& OutArray) const
+{
+	for (const FStreamingLevelActors& StreamingLevel : StreamingLevelLists)
+	{
+		StreamingLevel.ReplicationActorList.AppendToTArray(OutArray);
+	}
+}
+
+void FStreamingLevelActorListCollection::Log(FReplicationGraphDebugInfo& DebugInfo) const
+{
+	for (const FStreamingLevelActors& StreamingLevelList : StreamingLevelLists)
+	{
+		LogActorRepList(DebugInfo, StreamingLevelList.StreamingLevelName.ToString(), StreamingLevelList.ReplicationActorList);
+	}
+}
+
+// --------------------------------------------------------------------------------------------------------------------------------------------
+
+void UReplicationGraphNode_ActorList::NotifyAddNetworkActor(const FNewReplicatedActorInfo& ActorInfo)
+{
+	UE_CLOG(CVar_RepGraph_LogActorAdd>0, LogReplicationGraph, Display, TEXT("UReplicationGraphNode_ActorList::NotifyAddNetworkActor %s on %s."), *ActorInfo.Actor->GetFullName(), *GetPathName());
+
+	if (ActorInfo.StreamingLevelName == NAME_None)
+	{
+		if (CVar_RepGraph_Verify)
+		{
+			ensureMsgf(ReplicationActorList.Contains(ActorInfo.Actor) == false, TEXT("%s being added to %s twice!"), *GetActorRepListTypeDebugString(ActorInfo.Actor) );
+		}
+
+		ReplicationActorList.Add(ActorInfo.Actor);
+	}
+	else
+	{
+		StreamingLevelCollection.AddActor(ActorInfo);
+	}
+}
+
+void UReplicationGraphNode_ActorList::NotifyRemoveNetworkActor(const FNewReplicatedActorInfo& ActorInfo, bool bWarnIfNotFound)
+{
+	UE_CLOG(CVar_RepGraph_LogActorRemove>0, LogReplicationGraph, Display, TEXT("UReplicationGraphNode_ActorList::NotifyRemoveNetworkActor %s on %s."), *ActorInfo.Actor->GetFullName(), *GetPathName());
+
+	if (ActorInfo.StreamingLevelName == NAME_None)
+	{
+		if (!ReplicationActorList.Remove(ActorInfo.Actor) && bWarnIfNotFound)
+		{
+			UE_LOG(LogReplicationGraph, Warning, TEXT("Attempted to remove %s from list %s but it was not found. (StreamingLevelName == NAME_None)"), *GetActorRepListTypeDebugString(ActorInfo.Actor), *GetFullName());
+		}
+
+		if (CVar_RepGraph_Verify)
+		{
+			ensureMsgf(ReplicationActorList.Contains(ActorInfo.Actor) == false, TEXT("Actor %s is still in %s after removal"), *GetActorRepListTypeDebugString(ActorInfo.Actor), *GetPathName());
+		}
+	}
+	else
+	{
+		StreamingLevelCollection.RemoveActor(ActorInfo, bWarnIfNotFound, this);
+	}
+}
+	
+void UReplicationGraphNode_ActorList::NotifyResetAllNetworkActors()
+{
+	ReplicationActorList.Reset();
+	StreamingLevelCollection.Reset();
+}
+
+void UReplicationGraphNode_ActorList::GatherActorListsForConnection(const FConnectionGatherActorListParameters& Params)
+{
+	Params.OutGatheredReplicationLists.AddReplicationActorList(ReplicationActorList);
+	StreamingLevelCollection.Gather(Params);	
+}
+
+void UReplicationGraphNode_ActorList::DeepCopyActorListsFrom(const UReplicationGraphNode_ActorList* Source)
+{
+	if (Source->ReplicationActorList.Num() > 0)
+	{
+		ReplicationActorList.CopyContentsFrom(Source->ReplicationActorList);
+	}
+
+	StreamingLevelCollection.DeepCopyFrom(Source->StreamingLevelCollection);
+}
+
+void UReplicationGraphNode_ActorList::GetAllActorsInNode_Debugging(TArray<FActorRepListType>& OutArray) const
+{
+	ReplicationActorList.AppendToTArray(OutArray);
+	StreamingLevelCollection.GetAll_Debug(OutArray);
+}
+
+void UReplicationGraphNode_ActorList::LogNode(FReplicationGraphDebugInfo& DebugInfo, const FString& NodeName) const
+{
+	DebugInfo.Log(NodeName);
+
+	DebugInfo.PushIndent();
+	LogActorRepList(DebugInfo, TEXT("World"), ReplicationActorList);
+	StreamingLevelCollection.Log(DebugInfo);
+	DebugInfo.PopIndent();
+}
+
+// --------------------------------------------------------------------------------------------------------------------------------------------
+
+int32 UReplicationGraphNode_ActorListFrequencyBuckets::DefaultNumBuckets = 3;
+int32 UReplicationGraphNode_ActorListFrequencyBuckets::DefaultListSize = 12;
+TArray<UReplicationGraphNode_ActorListFrequencyBuckets::FBucketThresholds, TInlineAllocator<4>> UReplicationGraphNode_ActorListFrequencyBuckets::DefaultBucketThresholds;
+
+void UReplicationGraphNode_ActorListFrequencyBuckets::NotifyAddNetworkActor(const FNewReplicatedActorInfo& ActorInfo)
+{
+	UE_CLOG(CVar_RepGraph_LogActorAdd>0, LogReplicationGraph, Display, TEXT("UReplicationGraphNode_ActorListFrequencyBuckets::NotifyAddNetworkActor %s on %s."), *ActorInfo.Actor->GetFullName(), *GetPathName());
+
+	if (ActorInfo.StreamingLevelName == NAME_None)
+	{
+		// Add to smallest bucket
+		FActorRepListRefView* BestList = nullptr;
+		int32 LeastNum = INT_MAX;
+		for (FActorRepListRefView& List : NonStreamingCollection)
+		{
+			if (List.Num() < LeastNum)
+			{
+				BestList = &List;
+				LeastNum = List.Num();
+			}
+
+			if (CVar_RepGraph_Verify)
+			{
+				ensureMsgf(List.Contains(ActorInfo.Actor) == false, TEXT("%s being added to %s twice!"), *GetActorRepListTypeDebugString(ActorInfo.Actor) );
+			}
+		}
+
+		repCheck(BestList != nullptr);
+		BestList->Add(ActorInfo.Actor);
+		TotalNumNonStreamingActors++;
+		CheckRebalance();
+	}
+	else
+	{
+		StreamingLevelCollection.AddActor(ActorInfo);
+	}
+}
+
+void UReplicationGraphNode_ActorListFrequencyBuckets::NotifyRemoveNetworkActor(const FNewReplicatedActorInfo& ActorInfo, bool bWarnIfNotFound)
+{
+	UE_CLOG(CVar_RepGraph_LogActorRemove>0, LogReplicationGraph, Display, TEXT("UReplicationGraphNode_ActorListFrequencyBuckets::NotifyRemoveNetworkActor %s on %s."), *ActorInfo.Actor->GetFullName(), *GetPathName());
+
+	if (ActorInfo.StreamingLevelName == NAME_None)
+	{
+		bool bFound = false;
+		for (FActorRepListRefView& List : NonStreamingCollection)
+		{
+			if (List.Remove(ActorInfo.Actor))
+			{
+				TotalNumNonStreamingActors--;
+				CheckRebalance();
+
+				if (!CVar_RepGraph_Verify)
+				{
+					// Eary out if we dont have to verify
+					return;
+				}
+
+				if (bFound)
+				{
+					// We already removed this actor so this is a dupe!
+					repCheck(CVar_RepGraph_Verify);
+					ensureMsgf(false, TEXT("Actor %s is still in %s after removal"), *GetActorRepListTypeDebugString(ActorInfo.Actor), *GetPathName());
+				}
+
+				bFound = true;
+			}
+		}
+
+		if (!bFound && bWarnIfNotFound)
+		{
+			UE_LOG(LogReplicationGraph, Warning, TEXT("Attempted to remove %s from list %s but it was not found. (StreamingLevelName == NAME_None)"), *GetActorRepListTypeDebugString(ActorInfo.Actor), *GetFullName());
+		}
+	}
+	else
+	{
+		StreamingLevelCollection.RemoveActor(ActorInfo, bWarnIfNotFound, this);
+	}
+}
+	
+void UReplicationGraphNode_ActorListFrequencyBuckets::NotifyResetAllNetworkActors()
+{
+	for (FActorRepListRefView& List : NonStreamingCollection)
+	{
+		List.Reset();
+	}
+	StreamingLevelCollection.Reset();
+	TotalNumNonStreamingActors = 0;
+}
+
+void UReplicationGraphNode_ActorListFrequencyBuckets::GatherActorListsForConnection(const FConnectionGatherActorListParameters& Params)
+{
+	const int32 idx = Params.ReplicationFrameNum % NonStreamingCollection.Num();
+	Params.OutGatheredReplicationLists.AddReplicationActorList(NonStreamingCollection[idx]);
+	StreamingLevelCollection.Gather(Params);	
+}
+
+void UReplicationGraphNode_ActorListFrequencyBuckets::SetNonStreamingCollectionSize(const int32 NewSize)
+{
+	// Save everything off
+	static TArray<FActorRepListType> FullList;
+	FullList.Reset();
+
+	for (FActorRepListRefView& List : NonStreamingCollection)
+	{
+		List.AppendToTArray(FullList);
+	}
+
+	// Reset
+	NonStreamingCollection.SetNum(NewSize);
+	for (FActorRepListRefView& List : NonStreamingCollection)
+	{
+		List.Reset(DefaultListSize);
+	}
+
+	// Readd/Rebalance
+	for (int32 idx=0; idx < FullList.Num(); ++idx)
+	{
+		NonStreamingCollection[idx % NewSize].Add( FullList[idx] );
+	}
+}
+
+void UReplicationGraphNode_ActorListFrequencyBuckets::CheckRebalance()
+{
+	const int32 CurrentNumBuckets = NonStreamingCollection.Num();
+	int32 DesiredNumBuckets = CurrentNumBuckets;
+
+	for (const FBucketThresholds& Threshold : DefaultBucketThresholds)
+	{
+		if (TotalNumNonStreamingActors <= Threshold.MaxActors)
+		{
+			DesiredNumBuckets = Threshold.NumBuckets;
+			break;
+		}
+	}
+
+	if (DesiredNumBuckets != CurrentNumBuckets)
+	{
+		//UE_LOG(LogReplicationGraph, Display, TEXT("Rebalancing %s for %d buckets (%d total actors)"), *GetPathName(), DesiredNumBuckets, TotalNumNonStreamingActors);
+		SetNonStreamingCollectionSize(DesiredNumBuckets);
+	}
+	
+}
+
+void UReplicationGraphNode_ActorListFrequencyBuckets::GetAllActorsInNode_Debugging(TArray<FActorRepListType>& OutArray) const
+{
+	for (const FActorRepListRefView& List : NonStreamingCollection)
+	{
+		List.AppendToTArray(OutArray);
+	}
+	StreamingLevelCollection.GetAll_Debug(OutArray);
+}
+
+void UReplicationGraphNode_ActorListFrequencyBuckets::LogNode(FReplicationGraphDebugInfo& DebugInfo, const FString& NodeName) const
+{
+	DebugInfo.Log(NodeName);
+	DebugInfo.PushIndent();
+	int32 i=0;
+
+	for (const FActorRepListRefView& List : NonStreamingCollection)
+	{
+		LogActorRepList(DebugInfo, FString::Printf(TEXT("World Bucket %d"), ++i), List);
+	}
+	StreamingLevelCollection.Log(DebugInfo);
+	DebugInfo.PopIndent();
+}
+
+
+// --------------------------------------------------------------------------------------------------------------------------------------------
+
+void UReplicationGraphNode_ConnectionDormanyNode::GatherActorListsForConnection(const FConnectionGatherActorListParameters& Params)
+{
+	ConditionalGatherDormantActorsForConnection(ReplicationActorList, Params);
+	
+	for (int32 idx=StreamingLevelCollection.StreamingLevelLists.Num()-1; idx>=0; --idx)
+	{
+		FStreamingLevelActorListCollection::FStreamingLevelActors& StreamingList = StreamingLevelCollection.StreamingLevelLists[idx];
+		if (StreamingList.ReplicationActorList.Num() <= 0)
+		{
+			StreamingLevelCollection.StreamingLevelLists.RemoveAtSwap(idx, 1, false);
+			continue;
+		}
+
+		if (Params.CheckClientVisibilityForLevel(StreamingList.StreamingLevelName))
+		{
+			ConditionalGatherDormantActorsForConnection(StreamingList.ReplicationActorList, Params);
+		}
+		else
+		{
+			UE_LOG(LogReplicationGraph, Verbose, TEXT("Level Not Loaded %s. (Client has %d levels loaded)"), *StreamingList.StreamingLevelName.ToString(), Params.ClientVisibleLevelNamesRef.Num());
+		}
+	}
+}
+
+void UReplicationGraphNode_ConnectionDormanyNode::ConditionalGatherDormantActorsForConnection(FActorRepListRefView& ConnectionList, const FConnectionGatherActorListParameters& Params)
+{
+	FPerConnectionActorInfoMap& ConnectionActorInfoMap = Params.ConnectionManager.ActorInfoMap;
+
+	// We can trickle if the TrickelStartCounter is 0. (Just trying to give it a few frames to settle)
+	bool bShouldTrickle = TrickleStartCounter == 0;
+
+	for (int32 idx = ConnectionList.Num()-1; idx >= 0; --idx)
+	{
+		FActorRepListType Actor = ConnectionList[idx];
+		FConnectionReplicationActorInfo& ConnectionActorInfo = ConnectionActorInfoMap.FindOrAdd(Actor);
+		if (ConnectionActorInfo.bDormantOnConnection)
+		{
+			// He can be removed
+			ConnectionList.RemoveAtSwap(idx);
+			UE_CLOG(CVar_RepGraph_LogNetDormancyDetails > 0, LogReplicationGraph, Display, TEXT("GRAPH_DORMANCY: Actor %s is Dormant on %s. Removing from list. (%d elements left)"), *Actor->GetPathName(), *GetName(), ConnectionList.Num());
+			bShouldTrickle = false; // Dont trickle this frame because we are still encountering dormant actors
+		}
+		else if (CVar_RepGraph_TrickleDistCullOnDormanyNodes > 0 && bShouldTrickle)
+		{
+			ConnectionActorInfo.CullDistanceSquared = 0.f;
+			bShouldTrickle = false; // trickle one actor per frame
+		}
+	}
+
+	if (ConnectionList.Num() > 0)
+	{
+		Params.OutGatheredReplicationLists.AddReplicationActorList(ConnectionList);
+		
+		if (TrickleStartCounter > 0)
+		{
+			TrickleStartCounter--;		
+		}
+	}
+}
+
+bool ContainsReverse(const FActorRepListRefView& List, FActorRepListType Actor)
+{
+	for (int32 idx=List.Num()-1; idx >= 0; --idx)
+	{
+		if (List[idx] == Actor)
+			return true;
+	}
+
+	return false;
+}
+
+void UReplicationGraphNode_ConnectionDormanyNode::NotifyActorDormancyFlush(FActorRepListType Actor)
+{
+	FNewReplicatedActorInfo ActorInfo(Actor);
+
+	// Dormancy is flushed so we need to make sure this actor is on this connection specific node.
+	// Guard against dupes in the list. Sometimes actors flush multiple times in a row or back to back frames.
+	// 
+	// It may be better to track last flush frame on GlobalActorRepInfo? 
+	if (ActorInfo.StreamingLevelName == NAME_None)
+	{
+		if (!ContainsReverse(ReplicationActorList, Actor))
+		{
+			ReplicationActorList.Add(ActorInfo.Actor);
+		}
+	}
+	else
+	{
+		FStreamingLevelActorListCollection::FStreamingLevelActors* Item = StreamingLevelCollection.StreamingLevelLists.FindByKey(ActorInfo.StreamingLevelName);
+		if (!Item)
+		{
+			Item = new (StreamingLevelCollection.StreamingLevelLists) FStreamingLevelActorListCollection::FStreamingLevelActors(ActorInfo.StreamingLevelName);
+			Item->ReplicationActorList.Add(ActorInfo.Actor);
+
+		} else if(!ContainsReverse(Item->ReplicationActorList, Actor))
+		{
+			Item->ReplicationActorList.Add(ActorInfo.Actor);
+		}
+	}
+}
+
+// --------------------------------------------------------------------------------------------------------------------------------------------
+
+float UReplicationGraphNode_DormancyNode::MaxZForConnection = WORLD_MAX;
+
+void UReplicationGraphNode_DormancyNode::NotifyResetAllNetworkActors()
+{
+	if (GraphGlobals.IsValid())
+	{
+		// Unregister dormancy callbacks first
+		for (FActorRepListType& Actor : ReplicationActorList)
+		{
+			FGlobalActorReplicationInfo& GlobalInfo = GraphGlobals->GlobalActorReplicationInfoMap->Get(Actor);
+			GlobalInfo.Events.DormancyFlush.RemoveAll(this);
+		}
+	}
+
+	// Dump our global actor list
+	Super::NotifyResetAllNetworkActors();
+
+	// Reset the per connection nodes
+	for (auto& MapIt :  ConnectionNodes)
+	{
+		if (MapIt.Value)
+		{
+			MapIt.Value->NotifyResetAllNetworkActors();
+		}
+	}
+}
+
+void UReplicationGraphNode_DormancyNode::AddDormantActor(const FNewReplicatedActorInfo& ActorInfo, FGlobalActorReplicationInfo& GlobalInfo)
+{
+	Super::NotifyAddNetworkActor(ActorInfo);
+	
+	UE_CLOG(CVar_RepGraph_LogNetDormancyDetails > 0 && ConnectionNodes.Num() > 0, LogReplicationGraph, Display, TEXT("GRAPH_DORMANCY: AddDormantActor %s on %s. Adding to %d connection nodes."), *ActorInfo.Actor->GetPathName(), *GetName(), ConnectionNodes.Num());
+	
+	for (auto& MapIt : ConnectionNodes)
+	{
+		UReplicationGraphNode_ConnectionDormanyNode* Node = MapIt.Value;
+		Node->NotifyAddNetworkActor(ActorInfo);
+	}
+
+	// Tell us if this guy flushes net dormancy so we force him back on connection lists
+	GlobalInfo.Events.DormancyFlush.AddUObject(this, &UReplicationGraphNode_DormancyNode::OnActorDormancyFlush);
+}
+
+void UReplicationGraphNode_DormancyNode::RemoveDormantActor(const FNewReplicatedActorInfo& ActorInfo, FGlobalActorReplicationInfo& ActorRepInfo)
+{
+	UE_CLOG(CVar_RepGraph_LogActorRemove>0, LogReplicationGraph, Display, TEXT("UReplicationGraphNode_DormancyNode::RemoveDormantActor %s on %s. (%d connection nodes). ChildNodes: %d"), *GetNameSafe(ActorInfo.Actor), *GetPathName(), ConnectionNodes.Num(), AllChildNodes.Num());
+
+	Super::NotifyRemoveNetworkActor(ActorInfo);
+
+	ActorRepInfo.Events.DormancyFlush.RemoveAll(this);
+
+	// Update any connection specific nodes
+	for (auto& MapIt : ConnectionNodes)
+	{
+		UReplicationGraphNode_ConnectionDormanyNode* Node = MapIt.Value;
+		Node->NotifyRemoveNetworkActor(ActorInfo, false); // Don't warn if not found, the node may have removed the actor itself. Not wortht he extra bookkeeping to skip the call.
+	}
+}
+
+void UReplicationGraphNode_DormancyNode::GatherActorListsForConnection(const FConnectionGatherActorListParameters& Params)
+{
+	if (Params.Viewer.ViewLocation.Z > MaxZForConnection)
+	{
+		return;
+	}
+
+	UReplicationGraphNode_ConnectionDormanyNode** NodePtrPtr = ConnectionNodes.Find(&Params.ConnectionManager);
+	UReplicationGraphNode_ConnectionDormanyNode* ConnectionNode = nullptr;
+	if (!NodePtrPtr)
+	{
+		// We dont have a per-connection node for this connection, so create one and copy over contents
+		ConnectionNode = CreateChildNode<UReplicationGraphNode_ConnectionDormanyNode>();
+		ConnectionNodes.Add(&Params.ConnectionManager) = ConnectionNode;
+
+		// Copy our master lists tot he connection node
+		ConnectionNode->DeepCopyActorListsFrom(this);
+
+		UE_CLOG(CVar_RepGraph_LogNetDormancyDetails > 0, LogReplicationGraph, Display, TEXT("GRAPH_DORMANCY: First time seeing connection %s in node %s. Created ConnectionDormancyNode %s."), *Params.ConnectionManager.GetName(), *GetName(), *ConnectionNode->GetName());
+	}
+	else
+	{
+		ConnectionNode = *NodePtrPtr;
+	}
+
+	ConnectionNode->GatherActorListsForConnection(Params);
+}
+
+void UReplicationGraphNode_DormancyNode::OnActorDormancyFlush(FActorRepListType Actor, FGlobalActorReplicationInfo& GlobalInfo)
+{
+	if (CVar_RepGraph_Verify)
+	{
+		FNewReplicatedActorInfo ActorInfo(Actor);
+		if (ActorInfo.StreamingLevelName == NAME_None)
+		{
+			ensureMsgf(ReplicationActorList.Contains(Actor), TEXT("UReplicationGraphNode_DormancyNode::OnActorDormancyFlush %s not present in %s actor lists!"), *Actor->GetPathName(), *GetPathName());
+		}
+		else
+		{
+			if (FStreamingLevelActorListCollection::FStreamingLevelActors* Item = StreamingLevelCollection.StreamingLevelLists.FindByKey(ActorInfo.StreamingLevelName))
+			{
+				ensureMsgf(Item->ReplicationActorList.Contains(Actor), TEXT("UReplicationGraphNode_DormancyNode::OnActorDormancyFlush %s not present in %s actor lists! Streaming Level: %s"), *GetActorRepListTypeDebugString(Actor), *GetPathName(), *ActorInfo.StreamingLevelName.ToString());
+			}
+		}
+	}
+
+	// -------------------
+		
+	UE_CLOG(CVar_RepGraph_LogNetDormancyDetails > 0 && ConnectionNodes.Num() > 0, LogReplicationGraph, Display, TEXT("GRAPH_DORMANCY: Actor %s Flushed Dormancy. %s. Refreshing all %d connection nodes."), *Actor->GetPathName(), *GetName(), ConnectionNodes.Num());
+
+	for (auto& MapIt : ConnectionNodes)
+	{
+		UReplicationGraphNode_ConnectionDormanyNode* Node = MapIt.Value;
+		Node->NotifyActorDormancyFlush(Actor);
+	}
+}
+// --------------------------------------------------------------------------------------------------------------------------------------------
+
+void UReplicationGraphNode_GridCell::NotifyResetAllNetworkActors()
+{
+	Super::NotifyResetAllNetworkActors();
+	if (DynamicNode)
+	{
+		DynamicNode->NotifyResetAllNetworkActors();
+	}
+	if (DormancyNode)
+	{
+		DormancyNode->NotifyResetAllNetworkActors();
+	}
+}
+
+void UReplicationGraphNode_GridCell::GatherActorListsForConnection(const FConnectionGatherActorListParameters& Params)
+{
+	// Static actors
+	Super::GatherActorListsForConnection(Params);
+
+	// Dynamic actors
+	if (DynamicNode)
+	{
+		DynamicNode->GatherActorListsForConnection(Params);
+	}
+
+	// Dormancy nodes
+	if (DormancyNode)
+	{
+		DormancyNode->GatherActorListsForConnection(Params);
+	}
+}
+
+void UReplicationGraphNode_GridCell::AddStaticActor(const FNewReplicatedActorInfo& ActorInfo, FGlobalActorReplicationInfo& GlobalInfo, bool bParentNodeHandlesDormancyChange)
+{
+	if (GlobalInfo.bWantsToBeDormant)
+	{
+		// Pass to dormancy node
+		GetDormancyNode()->AddDormantActor(ActorInfo, GlobalInfo);
+	}
+	else
+	{	
+		// Put him in our non dormancy list
+		Super::NotifyAddNetworkActor(ActorInfo);
+	}
+
+	// We need to be told if this actor changes dormancy so we can move him between nodes. Unless our parent is going to do it.
+	if (!bParentNodeHandlesDormancyChange)
+	{
+		GlobalInfo.Events.DormancyChange.AddUObject(this, &UReplicationGraphNode_GridCell::OnNetDormancyChange);
+	}
+}
+
+void UReplicationGraphNode_GridCell::AddDynamicActor(const FNewReplicatedActorInfo& ActorInfo)
+{
+	GetDynamicNode()->NotifyAddNetworkActor(ActorInfo);
+}
+
+void UReplicationGraphNode_GridCell::RemoveStaticActor(const FNewReplicatedActorInfo& ActorInfo, FGlobalActorReplicationInfo& ActorRepInfo, bool bWasAddedAsDormantActor)
+{
+	UE_CLOG(CVar_RepGraph_LogActorRemove>0, LogReplicationGraph, Display, TEXT("UReplicationGraphNode_Simple2DSpatializationLeaf::RemoveStaticActor %s on %s"), *ActorInfo.Actor->GetPathName(), *GetPathName());
+
+	if (bWasAddedAsDormantActor)
+	{
+		GetDormancyNode()->RemoveDormantActor(ActorInfo, ActorRepInfo);
+	}
+	else
+	{	
+		Super::NotifyRemoveNetworkActor(ActorInfo);
+	}
+	
+	ActorRepInfo.Events.DormancyChange.RemoveAll(this);
+}
+
+void UReplicationGraphNode_GridCell::RemoveDynamicActor(const FNewReplicatedActorInfo& ActorInfo)
+{
+	UE_CLOG(CVar_RepGraph_LogActorRemove>0, LogReplicationGraph, Display, TEXT("UReplicationGraphNode_Simple2DSpatializationLeaf::RemoveDynamicActor %s on %s"), *ActorInfo.Actor->GetPathName(), *GetPathName());
+
+	GetDynamicNode()->NotifyRemoveNetworkActor(ActorInfo);
+}
+
+void UReplicationGraphNode_GridCell::ConditionalCopyDormantActors(FActorRepListRefView& FromList, UReplicationGraphNode_DormancyNode* ToNode)
+{
+	if (GraphGlobals.IsValid())
+	{
+		for (int32 idx = FromList.Num()-1; idx >= 0; --idx)
+		{
+			FActorRepListType Actor = FromList[idx];
+			FGlobalActorReplicationInfo& GlobalInfo = GraphGlobals->GlobalActorReplicationInfoMap->Get(Actor);
+			if (GlobalInfo.bWantsToBeDormant)
+			{
+				ToNode->NotifyAddNetworkActor(FNewReplicatedActorInfo(Actor));
+				FromList.RemoveAtSwap(idx);
+			}
+		}
+	}
+}
+
+void UReplicationGraphNode_GridCell::OnNetDormancyChange(FActorRepListType Actor, FGlobalActorReplicationInfo& GlobalInfo, ENetDormancy NewValue, ENetDormancy OldValue)
+{
+	UE_CLOG(CVar_RepGraph_LogActorRemove>0, LogReplicationGraph, Display, TEXT("UReplicationGraphNode_Simple2DSpatializationLeaf::OnNetDormancyChange. %s on %s. Old: %d, New: %d"), *Actor->GetPathName(), *GetPathName(), NewValue, OldValue);
+
+	const bool bCurrentDormant = NewValue > DORM_Awake;
+	const bool bPreviousDormant = OldValue > DORM_Awake;
+
+	if (!bCurrentDormant && bPreviousDormant)
+	{
+		// Actor is now awake, remove from dormany node and add to non dormany list
+		FNewReplicatedActorInfo ActorInfo(Actor);
+		GetDormancyNode()->RemoveDormantActor(ActorInfo, GlobalInfo);
+		Super::NotifyAddNetworkActor(ActorInfo);
+	}
+	else if (bCurrentDormant && !bPreviousDormant)
+	{
+		// Actor is now dormant, remove from non dormant list, add to dormant node
+		FNewReplicatedActorInfo ActorInfo(Actor);
+		Super::NotifyRemoveNetworkActor(ActorInfo);
+		GetDormancyNode()->AddDormantActor(ActorInfo, GlobalInfo);
+	}
+}
+
+
+UReplicationGraphNode_ActorListFrequencyBuckets* UReplicationGraphNode_GridCell::GetDynamicNode()
+{
+	if (DynamicNode == nullptr)
+	{
+		DynamicNode = CreateChildNode<UReplicationGraphNode_ActorListFrequencyBuckets>();
+	}
+
+	return DynamicNode;
+}
+
+UReplicationGraphNode_DormancyNode* UReplicationGraphNode_GridCell::GetDormancyNode()
+{
+	if (DormancyNode == nullptr)
+	{
+		DormancyNode = CreateChildNode<UReplicationGraphNode_DormancyNode>();
+	}
+
+	return DormancyNode;
+}
+
+void UReplicationGraphNode_GridCell::GetAllActorsInNode_Debugging(TArray<FActorRepListType>& OutArray) const
+{
+	Super::GetAllActorsInNode_Debugging(OutArray);
+	if (DynamicNode)
+	{
+		DynamicNode->GetAllActorsInNode_Debugging(OutArray);
+	}
+	if (DormancyNode)
+	{
+		DormancyNode->GetAllActorsInNode_Debugging(OutArray);
+	}
+}
+
+int32 CVar_RepGraph_DebugNextNewActor = 0;
+static FAutoConsoleVariableRef CVarRepGraphDebugNextActor(TEXT("Net.RepGraph.Spatial.DebugNextNewActor"), CVar_RepGraph_DebugNextNewActor, TEXT(""), ECVF_Default );
+
+// -------------------------------------------------------
+
+UReplicationGraphNode_GridSpatialization2D::UReplicationGraphNode_GridSpatialization2D()
+	: CellSize(0.f)
+	, SpatialBias(ForceInitToZero)
+{
+	bRequiresPrepareForReplicationCall = true;
+
+	SetCreateChildNodeFunc([](UReplicationGraphNode* Parent)
+	{
+		return Parent->CreateChildNode<UReplicationGraphNode_GridCell>();
+	});
+}
+
+void UReplicationGraphNode_GridSpatialization2D::NotifyAddNetworkActor(const FNewReplicatedActorInfo& ActorInfo)
+{
+	ensureAlwaysMsgf(false, TEXT("UReplicationGraphNode_GridSpatialization2D::NotifyAddNetworkActor should not be called directly"));
+}
+
+void UReplicationGraphNode_GridSpatialization2D::NotifyRemoveNetworkActor(const FNewReplicatedActorInfo& ActorInfo, bool bWarnIfNotFound)
+{
+	ensureAlwaysMsgf(false, TEXT("UReplicationGraphNode_GridSpatialization2D::NotifyRemoveNetworkActor should not be called directly"));
+}
+
+void UReplicationGraphNode_GridSpatialization2D::AddActor_Dormancy(const FNewReplicatedActorInfo& ActorInfo, FGlobalActorReplicationInfo& ActorRepInfo)
+{
+	UE_CLOG(CVar_RepGraph_LogActorRemove>0, LogReplicationGraph, Display, TEXT("UReplicationGraphNode_GridSpatialization2D::AddActor_Dormancy %s on %s"), *ActorInfo.Actor->GetFullName(), *GetPathName());
+
+	if (ActorRepInfo.bWantsToBeDormant)
+	{
+		AddActorInternal_Static(ActorInfo, ActorRepInfo, true);
+	}
+	else
+	{
+		AddActorInternal_Dynamic(ActorInfo);
+	}
+
+	// Tell us if dormancy changes for this actor because then we need to move it. Note we don't care about Flushing.
+	ActorRepInfo.Events.DormancyChange.AddUObject(this, &UReplicationGraphNode_GridSpatialization2D::OnNetDormancyChange);
+}
+
+void UReplicationGraphNode_GridSpatialization2D::RemoveActor_Static(const FNewReplicatedActorInfo& ActorInfo)
+{
+	UE_CLOG(CVar_RepGraph_LogActorRemove>0, LogReplicationGraph, Display, TEXT("UReplicationGraphNode_GridSpatialization2D::RemoveActor_Static %s on %s"), *ActorInfo.Actor->GetFullName(), *GetPathName());
+
+	if (GraphGlobals.IsValid())
+	{
+		FGlobalActorReplicationInfo& GlobalInfo = GraphGlobals->GlobalActorReplicationInfoMap->Get(ActorInfo.Actor);
+		RemoveActorInternal_Static(ActorInfo, GlobalInfo, GlobalInfo.bWantsToBeDormant); 
+	}
+}
+
+void UReplicationGraphNode_GridSpatialization2D::RemoveActor_Dormancy(const FNewReplicatedActorInfo& ActorInfo)
+{
+	UE_CLOG(CVar_RepGraph_LogActorRemove>0, LogReplicationGraph, Display, TEXT("UReplicationGraphNode_GridSpatialization2D::RemoveActor_Dormancy %s on %s"), *ActorInfo.Actor->GetFullName(), *GetPathName());
+
+	if (GraphGlobals.IsValid())
+	{
+		FGlobalActorReplicationInfo& ActorRepInfo = GraphGlobals->GlobalActorReplicationInfoMap->Get(ActorInfo.Actor);
+		if (ActorRepInfo.bWantsToBeDormant)
+		{
+			RemoveActorInternal_Static(ActorInfo, ActorRepInfo, true);
+		}
+		else
+		{
+			RemoveActorInternal_Dynamic(ActorInfo);
+		}
+	}
+}
+
+void UReplicationGraphNode_GridSpatialization2D::AddActorInternal_Dynamic(const FNewReplicatedActorInfo& ActorInfo)
+{
+#if !(UE_BUILD_SHIPPING || UE_BUILD_TEST)
+	if (ActorInfo.Actor->bAlwaysRelevant)
+	{
+		UE_LOG(LogReplicationGraph, Warning, TEXT("Always relevant actor being added to spatialized graph node. %s"), *GetNameSafe(ActorInfo.Actor));
+		return;
+	}
+#endif
+
+	UE_CLOG(CVar_RepGraph_LogActorRemove>0, LogReplicationGraph, Display, TEXT("UReplicationGraph::AddActorInternal_Dynamic %s"), *ActorInfo.Actor->GetFullName());
+
+	DynamicSpatializedActors.Emplace(ActorInfo.Actor, ActorInfo);
+}
+
+void UReplicationGraphNode_GridSpatialization2D::AddActorInternal_Static(const FNewReplicatedActorInfo& ActorInfo, FGlobalActorReplicationInfo& ActorRepInfo, bool bDormancyDriven)
+{
+	AActor* Actor = ActorInfo.Actor;
+	const FVector Location3D = Actor->GetActorLocation();
+	ActorRepInfo.WorldLocation = Location3D;
+
+	if (CVar_RepGraph_LogActorAdd)
+	{
+		UE_LOG(LogReplicationGraph, Display, TEXT("UReplicationGraphNode_GridSpatialization2D::AddActorInternal_Static placing %s into static grid at %s"), *Actor->GetPathName(), *ActorRepInfo.WorldLocation.ToString());
+	}
+		
+	if (SpatialBias.X > Location3D.X || SpatialBias.Y > Location3D.Y)
+	{
+		HandleActorOutOfSpatialBounds(Actor, Location3D, true);
+	}
+
+	StaticSpatializedActors.Emplace(Actor, FCachedStaticActorInfo(ActorInfo, bDormancyDriven));
+
+	// Only put in cell right now if we aren't needing to rebuild the whole grid
+	if (!bNeedsRebuild)
+	{
+		PutStaticActorIntoCell(ActorInfo, ActorRepInfo, bDormancyDriven);
+	}
+}
+
+void UReplicationGraphNode_GridSpatialization2D::RemoveActorInternal_Dynamic(const FNewReplicatedActorInfo& ActorInfo)
+{
+	if (FCachedDynamicActorInfo* DynamicActorInfo = DynamicSpatializedActors.Find(ActorInfo.Actor))
+	{
+		if (DynamicActorInfo->CellInfo.IsValid())
+		{
+			GetGridNodesForActor(ActorInfo.Actor, DynamicActorInfo->CellInfo, GatheredNodes);
+			for (UReplicationGraphNode_GridCell* Node : GatheredNodes)
+			{
+				Node->RemoveDynamicActor(ActorInfo);
+			}
+		}
+		DynamicSpatializedActors.Remove(ActorInfo.Actor);
+	}
+	else
+	{
+		UE_LOG(LogReplicationGraph, Warning, TEXT("UReplicationGraphNode_Simple2DSpatialization::RemoveActorInternal_Dynamic attempted remove %s from streaming dynamic list but it was not there."), *GetActorRepListTypeDebugString(ActorInfo.Actor));
+		if (StaticSpatializedActors.Remove(ActorInfo.Actor) > 0)
+		{
+			UE_LOG(LogReplicationGraph, Warning, TEXT("   It was in StaticSpatializedActors!"));
+		}
+	}
+}
+
+void UReplicationGraphNode_GridSpatialization2D::RemoveActorInternal_Static(const FNewReplicatedActorInfo& ActorInfo, FGlobalActorReplicationInfo& ActorRepInfo, bool bWasAddedAsDormantActor)
+{
+	if (StaticSpatializedActors.Remove(ActorInfo.Actor) <= 0)
+	{
+		UE_LOG(LogReplicationGraph, Warning, TEXT("UReplicationGraphNode_Simple2DSpatialization::RemoveActorInternal_Static attempted remove %s from static list but it was not there."), *GetActorRepListTypeDebugString(ActorInfo.Actor));
+		if(DynamicSpatializedActors.Remove(ActorInfo.Actor) > 0)
+		{
+			UE_LOG(LogReplicationGraph, Warning, TEXT("   It was in DynamicStreamingSpatializedActors!"));
+		}
+	}
+
+	// Remove it from the actual node it should still be in. Note that even if the actor did move in between this and the last replication frame, the FGlobalActorReplicationInfo would not have been updated
+	GetGridNodesForActor(ActorInfo.Actor, ActorRepInfo, GatheredNodes);
+	for (UReplicationGraphNode_GridCell* Node : GatheredNodes)
+	{
+		Node->RemoveStaticActor(ActorInfo, ActorRepInfo, bWasAddedAsDormantActor);
+	}
+	
+	if (CVar_RepGraph_Verify)
+	{
+		// Verify this actor is in no nodes. This is pretty slow!
+		TArray<AActor*> AllActors;
+		for (auto& InnerArray : Grid)
+		{
+			for (UReplicationGraphNode_GridCell* N : InnerArray)
+			{
+				if (N)
+				{
+					AllActors.Reset();
+					N->GetAllActorsInNode_Debugging(AllActors);
+					
+					ensureMsgf(AllActors.Contains(ActorInfo.Actor) == false, TEXT("Actor still in a node after removal!. %s. Removal Location: %s"), *N->GetPathName(), *ActorRepInfo.WorldLocation.ToString());
+				}
+			}
+		}
+	}
+}
+
+void UReplicationGraphNode_GridSpatialization2D::OnNetDormancyChange(FActorRepListType Actor, FGlobalActorReplicationInfo& GlobalInfo, ENetDormancy NewValue, ENetDormancy OldValue)
+{
+	const bool bCurrentShouldBeStatic = NewValue > DORM_Awake;
+	const bool bPreviousShouldBeStatic = OldValue > DORM_Awake;
+
+	if (bCurrentShouldBeStatic && !bPreviousShouldBeStatic)
+	{
+		// Actor was dynamic and is now static. Remove from dynamic list and add to static.
+		FNewReplicatedActorInfo ActorInfo(Actor);
+		RemoveActorInternal_Dynamic(ActorInfo);
+		AddActorInternal_Static(ActorInfo, GlobalInfo, true);
+	}
+	else if (!bCurrentShouldBeStatic && bPreviousShouldBeStatic)
+	{
+		FNewReplicatedActorInfo ActorInfo(Actor);
+		RemoveActorInternal_Static(ActorInfo, GlobalInfo, true); // This is why we need the 3rd bool parameter: this actor was placed as dormant (and it no longer is at the moment of this callback)
+		AddActorInternal_Dynamic(ActorInfo);
+	}
+}
+
+void UReplicationGraphNode_GridSpatialization2D::NotifyResetAllNetworkActors()
+{
+	StaticSpatializedActors.Reset();
+	DynamicSpatializedActors.Reset();
+	Super::NotifyResetAllNetworkActors();
+}
+
+void UReplicationGraphNode_GridSpatialization2D::PutStaticActorIntoCell(const FNewReplicatedActorInfo& ActorInfo, FGlobalActorReplicationInfo& ActorRepInfo, bool bDormancyDriven)
+{
+	GetGridNodesForActor(ActorInfo.Actor, ActorRepInfo, GatheredNodes);
+	for (UReplicationGraphNode_GridCell* Node : GatheredNodes)
+	{
+		Node->AddStaticActor(ActorInfo, ActorRepInfo, bDormancyDriven);
+	}
+}
+
+void UReplicationGraphNode_GridSpatialization2D::GetGridNodesForActor(FActorRepListType Actor, const FGlobalActorReplicationInfo& ActorRepInfo, TArray<UReplicationGraphNode_GridCell*>& OutNodes)
+{
+	RG_QUICK_SCOPE_CYCLE_COUNTER(UReplicationGraphNode_GridSpatialization2D_GetGridNodesForActor);
+	GetGridNodesForActor(Actor, GetCellInfoForActor(Actor, ActorRepInfo.WorldLocation, ActorRepInfo.Settings.CullDistanceSquared), OutNodes);
+}
+
+UReplicationGraphNode_GridSpatialization2D::FActorCellInfo UReplicationGraphNode_GridSpatialization2D::GetCellInfoForActor(FActorRepListType Actor, const FVector& Location3D, float CullDistanceSquared)
+{
+#if !(UE_BUILD_SHIPPING || UE_BUILD_TEST)
+	if (CullDistanceSquared <= 0.f)
+	{
+		UE_LOG(LogReplicationGraph, Warning, TEXT("::GetGridNodesForActor called on %s when its CullDistanceSquared = %.2f. (Must be > 0)"), *GetActorRepListTypeDebugString(Actor), CullDistanceSquared);
+	}
+#endif
+
+	FActorCellInfo CellInfo;
+	const float LocationBiasX = (Location3D.X - SpatialBias.X);
+	const float LocationBiasY = (Location3D.Y - SpatialBias.Y);
+
+	const float Dist = FMath::Sqrt(CullDistanceSquared);	 // Fixme Sqrt
+	const float MinX = LocationBiasX - Dist;
+	const float MinY = LocationBiasY - Dist;
+	const float MaxX = LocationBiasX + Dist;
+	const float MaxY = LocationBiasY + Dist;
+
+	CellInfo.StartX = FMath::Max<int32>(0, MinX / CellSize);
+	CellInfo.StartY = FMath::Max<int32>(0, MinY / CellSize);
+
+	CellInfo.EndX = FMath::Max<int32>(0, MaxX / CellSize);
+	CellInfo.EndY = FMath::Max<int32>(0, MaxY / CellSize);
+	return CellInfo;
+}
+
+void UReplicationGraphNode_GridSpatialization2D::GetGridNodesForActor(FActorRepListType Actor, const UReplicationGraphNode_GridSpatialization2D::FActorCellInfo& CellInfo, TArray<UReplicationGraphNode_GridCell*>& OutNodes)
+{
+	if (!ensure(CellInfo.IsValid()))
+	{
+		return;
+	}
+
+	OutNodes.Reset();
+
+	const int32 StartX = CellInfo.StartX;
+	const int32 StartY = CellInfo.StartY;
+	const int32 EndX = CellInfo.EndX;
+	const int32 EndY = CellInfo.EndY;
+
+	if (Grid.Num() <= EndX)
+	{
+		Grid.SetNum(EndX+1);
+	}
+
+	for (int32 X = StartX; X <= EndX; X++)
+	{
+		TArray<UReplicationGraphNode_GridCell*>& GridY = Grid[X];
+		if (GridY.Num() <= EndY)
+		{
+			GridY.SetNum(EndY+1);
+		}
+
+		for (int32 Y = StartY; Y <= EndY; Y++)
+		{
+			UReplicationGraphNode_GridCell*& NodePtr = GridY[Y];
+			if (NodePtr == nullptr)
+			{
+				NodePtr = CastChecked<UReplicationGraphNode_GridCell>(CreateChildNode());
+			}
+
+			OutNodes.Add(NodePtr);
+		}
+	}
+
+/*
+#if !(UE_BUILD_SHIPPING || UE_BUILD_TEST)
+	if (DebugActorNames.Num() > 0)
+	{
+		if (DebugActorNames.ContainsByPredicate([&](const FString DebugName) { return Actor->GetName().Contains(DebugName); }))
+		{
+			UE_LOG(LogReplicationGraph, Display, TEXT("Adding Actor %s. WorldLocation (Cached): %s. WorldLocation (AActor): %s. Buckets: %d/%d. SpatialBias: %s"), *Actor->GetName(), *ActorRepInfo.WorldLocation.ToString(), *Actor->GetActorLocation().ToString(), BucketX, BucketY, *SpatialBias.ToString());
+		}
+	}
+#endif
+*/
+}
+
+void UReplicationGraphNode_GridSpatialization2D::HandleActorOutOfSpatialBounds(AActor* Actor, const FVector& Location3D, const bool bStaticActor)
+{
+	// Don't rebuild spatialization for blacklisted actors. They will just get clamped to the grid.
+	if (RebuildSpatialBlacklistMap.Get(Actor->GetClass()) != nullptr)
+	{
+		return;
+	}
+
+	const bool bOldNeedRebuild = bNeedsRebuild;
+	if (SpatialBias.X > Location3D.X)
+	{
+		bNeedsRebuild = true;
+		SpatialBias.X = Location3D.X - (CellSize / 2.f);
+	}
+	if (SpatialBias.Y > Location3D.Y)
+	{
+		bNeedsRebuild = true;
+		SpatialBias.Y = Location3D.Y - (CellSize / 2.f);
+	}
+
+	if (bNeedsRebuild && !bOldNeedRebuild)
+	{
+		UE_LOG(LogReplicationGraph, Warning, TEXT("Spatialization Rebuild caused by: %s at %s. New Bias: %s. IsStatic: %d"), *Actor->GetPathName(), *Location3D.ToString(), *SpatialBias.ToString(), (int32)bStaticActor);
+	}
+}
+
+int32 CVar_RepGraph_Spatial_PauseDynamic = 0;
+static FAutoConsoleVariableRef CVarRepSpatialPauseDynamic(TEXT("Net.RepGraph.Spatial.PauseDynamic"), CVar_RepGraph_Spatial_PauseDynamic, TEXT("Pauses updating dynamic actor positions in the spatialization nodes."), ECVF_Default );
+
+int32 CVar_RepGraph_Spatial_DebugDynamic = 0;
+static FAutoConsoleVariableRef CVarRepGraphSpatialDebugDynamic(TEXT("Net.RepGraph.Spatial.DebugDynamic"), CVar_RepGraph_Spatial_DebugDynamic, TEXT("Prints debug info whenever dynamic actors changes spatial cells"), ECVF_Default );
+
+int32 CVar_RepGraph_Spatial_BiasCreep = 0.f;
+static FAutoConsoleVariableRef CVarRepGraphSpatialBiasCreep(TEXT("Net.RepGraph.Spatial.BiasCreep"), CVar_RepGraph_Spatial_BiasCreep, TEXT("Changes bias each frame by this much and force rebuld. For stress test debugging"), ECVF_Default );
+
+
+void UReplicationGraphNode_GridSpatialization2D::PrepareForReplication()
+{
+	RG_QUICK_SCOPE_CYCLE_COUNTER(UReplicationGraphNode_GridSpatialization2D_PrepareForReplication);
+
+	FGlobalActorReplicationInfoMap* GlobalRepMap = GraphGlobals.IsValid() ? GraphGlobals->GlobalActorReplicationInfoMap : nullptr;
+
+#if !(UE_BUILD_SHIPPING || UE_BUILD_TEST)
+	if (CVar_RepGraph_Spatial_BiasCreep != 0.f)
+	{
+		SpatialBias.X += CVar_RepGraph_Spatial_BiasCreep;
+		SpatialBias.Y += CVar_RepGraph_Spatial_BiasCreep;
+		bNeedsRebuild = true;
+	}
+
+
+	if (CVar_RepGraph_Spatial_PauseDynamic == 0)
+#endif
+	{
+		RG_QUICK_SCOPE_CYCLE_COUNTER(UReplicationGraphNode_GridSpatialization2D_BuildDynamic);
+
+		for (auto& MapIt : DynamicSpatializedActors)
+		{
+			FActorRepListType& DynamicActor = MapIt.Key;
+			FCachedDynamicActorInfo& DynamicActorInfo = MapIt.Value;
+			FActorCellInfo& PreviousCellInfo = DynamicActorInfo.CellInfo;
+			FNewReplicatedActorInfo& ActorInfo = DynamicActorInfo.ActorInfo;
+
+#if !(UE_BUILD_SHIPPING || UE_BUILD_TEST)
+			if (!IsActorValidForReplicationGather(DynamicActor))
+			{
+				UE_LOG(LogReplicationGraph, Warning, TEXT("UReplicationGraphNode_GridSpatialization2D::PrepareForReplication: Dynamic Actor no longer ready for replication"));
+				UE_LOG(LogReplicationGraph, Warning, TEXT("%s"), *GetNameSafe(DynamicActor));
+				continue;
+			}
+#endif
+
+			// Update location
+			FGlobalActorReplicationInfo& ActorRepInfo = GlobalRepMap->Get(DynamicActor);
+
+			// Check if this resets spatial bias
+			const FVector Location3D = DynamicActor->GetActorLocation();
+			ActorRepInfo.WorldLocation = Location3D;
+			
+			if (SpatialBias.X > Location3D.X || SpatialBias.Y > Location3D.Y)
+			{
+				HandleActorOutOfSpatialBounds(DynamicActor, Location3D, false);
+			}
+
+			if (!bNeedsRebuild)
+			{
+				// Get the new CellInfo
+				const FActorCellInfo NewCellInfo = GetCellInfoForActor(DynamicActor, Location3D, ActorRepInfo.Settings.CullDistanceSquared);
+
+				if (PreviousCellInfo.IsValid())
+				{
+					bool bDirty = false;
+
+					if (UNLIKELY(NewCellInfo.StartX > PreviousCellInfo.EndX || NewCellInfo.EndX < PreviousCellInfo.StartX ||
+							NewCellInfo.StartY > PreviousCellInfo.EndY || NewCellInfo.EndY < PreviousCellInfo.StartY))
+					{
+						// No longer intersecting, we just have to remove from all previous nodes and add to all new nodes
+						
+						bDirty = true;
+
+						GetGridNodesForActor(DynamicActor, PreviousCellInfo, GatheredNodes);
+						for (UReplicationGraphNode_GridCell* Node : GatheredNodes)
+						{
+							Node->RemoveDynamicActor(ActorInfo);
+						}
+
+						GetGridNodesForActor(DynamicActor, NewCellInfo, GatheredNodes);
+						for (UReplicationGraphNode_GridCell* Node : GatheredNodes)
+						{
+							Node->AddDynamicActor(ActorInfo);
+						}
+					}
+					else
+					{
+						// Some overlap so lets find out what cells need to be added or removed
+
+						if (PreviousCellInfo.StartX < NewCellInfo.StartX)
+						{
+							// We lost columns on the left side
+							bDirty = true;
+						
+							for (int32 X = PreviousCellInfo.StartX; X < NewCellInfo.StartX; ++X)
+							{
+								auto& GridX = GetGridX(X);
+								for (int32 Y = PreviousCellInfo.StartY; Y <= PreviousCellInfo.EndY; ++Y)
+								{
+									if (auto& Node = GetCell(GridX, Y))
+									{
+										Node->RemoveDynamicActor(ActorInfo);
+									}
+								}
+							}
+						}
+						else if(PreviousCellInfo.StartX > NewCellInfo.StartX)
+						{
+							// We added columns on the left side
+							bDirty = true;
+
+							for (int32 X = NewCellInfo.StartX; X < PreviousCellInfo.StartX; ++X)
+							{
+								auto& GridX = GetGridX(X);
+								for (int32 Y = NewCellInfo.StartY; Y <= NewCellInfo.EndY; ++Y)
+								{
+									GetLeafNode(GetCell(GridX,Y))->AddDynamicActor(ActorInfo);
+								}
+							}
+
+						}
+
+						if (PreviousCellInfo.EndX < NewCellInfo.EndX)
+						{
+							// We added columns on the right side
+							bDirty = true;
+
+							for (int32 X = PreviousCellInfo.EndX+1; X <= NewCellInfo.EndX; ++X)
+							{
+								auto& GridX = GetGridX(X);
+								for (int32 Y = NewCellInfo.StartY; Y <= NewCellInfo.EndY; ++Y)
+								{
+									GetLeafNode(GetCell(GridX,Y))->AddDynamicActor(ActorInfo);
+								}
+							}
+						
+						}
+						else if(PreviousCellInfo.EndX > NewCellInfo.EndX)
+						{
+							// We lost columns on the right side
+							bDirty = true;
+
+							for (int32 X = NewCellInfo.EndX+1; X <= PreviousCellInfo.EndX; ++X)
+							{
+								auto& GridX = GetGridX(X);
+								for (int32 Y = PreviousCellInfo.StartY; Y <= PreviousCellInfo.EndY; ++Y)
+								{
+									if (auto& Node = GetCell(GridX, Y))
+									{
+										Node->RemoveDynamicActor(ActorInfo);
+									}
+								}
+							}
+						}
+
+						// --------------------------------------------------
+
+						// We've handled left/right sides. So while handling top and bottom we only need to worry about this run of X cells
+						const int32 StartX = FMath::Max<int32>(NewCellInfo.StartX, PreviousCellInfo.StartX);
+						const int32 EndX = FMath::Min<int32>(NewCellInfo.EndX, PreviousCellInfo.EndX);
+
+						if (PreviousCellInfo.StartY < NewCellInfo.StartY)
+						{
+							// We lost rows on the top side
+							bDirty = true;
+							
+							for (int32 X = StartX; X <= EndX; ++X)
+							{
+								auto& GridX = GetGridX(X);
+								for (int32 Y = PreviousCellInfo.StartY; Y < NewCellInfo.StartY; ++Y)
+								{
+									if (auto& Node = GetCell(GridX, Y))
+									{
+										Node->RemoveDynamicActor(ActorInfo);
+									}
+								}
+							}
+						}
+						else if(PreviousCellInfo.StartY > NewCellInfo.StartY)
+						{
+							// We added rows on the top side
+							bDirty = true;
+							
+							for (int32 X = StartX; X <= EndX; ++X)
+							{
+								auto& GridX = GetGridX(X);
+								for (int32 Y = NewCellInfo.StartY; Y < PreviousCellInfo.StartY; ++Y)
+								{
+									GetLeafNode(GetCell(GridX,Y))->AddDynamicActor(ActorInfo);
+								}
+							}
+						}
+
+						if (PreviousCellInfo.EndY < NewCellInfo.EndY)
+						{
+							// We added rows on the bottom side
+							bDirty = true;
+							
+							for (int32 X = StartX; X <= EndX; ++X)
+							{
+								auto& GridX = GetGridX(X);
+								for (int32 Y = PreviousCellInfo.EndY+1; Y <= NewCellInfo.EndY; ++Y)
+								{
+									GetLeafNode(GetCell(GridX,Y))->AddDynamicActor(ActorInfo);
+								}
+							}
+						}
+						else if (PreviousCellInfo.EndY > NewCellInfo.EndY)
+						{
+							// We lost rows on the bottom side
+							bDirty = true;
+							
+							for (int32 X = StartX; X <= EndX; ++X)
+							{
+								auto& GridX = GetGridX(X);
+								for (int32 Y = NewCellInfo.EndY+1; Y <= PreviousCellInfo.EndY; ++Y)
+								{
+									if (auto& Node = GetCell(GridX, Y))
+									{
+										Node->RemoveDynamicActor(ActorInfo);
+									}
+								}
+							}
+						}
+					}
+
+					if (bDirty)
+					{
+#if !(UE_BUILD_SHIPPING || UE_BUILD_TEST)
+						if (CVar_RepGraph_Spatial_DebugDynamic)
+						{
+							auto CellInfoStr = [](const FActorCellInfo& CellInfo) { return FString::Printf(TEXT("[%d,%d]-[%d,%d]"), CellInfo.StartX, CellInfo.StartY, CellInfo.EndX, CellInfo.EndY); };
+							UE_LOG(LogReplicationGraph, Display, TEXT("%s moved cells. From %s to %s"), *GetActorRepListTypeDebugString(DynamicActor), *CellInfoStr(PreviousCellInfo), *CellInfoStr(NewCellInfo));
+
+							const int32 MinX = FMath::Min<int32>(PreviousCellInfo.StartX, NewCellInfo.StartX);
+							const int32 MinY = FMath::Min<int32>(PreviousCellInfo.StartY, NewCellInfo.StartY);
+							const int32 MaxX = FMath::Max<int32>(PreviousCellInfo.EndX, NewCellInfo.EndX);
+							const int32 MaxY = FMath::Max<int32>(PreviousCellInfo.EndY, NewCellInfo.EndY);
+
+							
+							for (int32 Y = MinY; Y <= MaxY; ++Y)
+							{
+								FString Str = FString::Printf(TEXT("[%d]   "), Y);
+								for (int32 X = MinX; X <= MaxX; ++X)
+								{
+									const bool bShouldBeInOld = (X >= PreviousCellInfo.StartX && X <= PreviousCellInfo.EndX) && (Y >= PreviousCellInfo.StartY && Y <= PreviousCellInfo.EndY);
+									const bool bShouldBeInNew = (X >= NewCellInfo.StartX && X <= NewCellInfo.EndX) && (Y >= NewCellInfo.StartY && Y <= NewCellInfo.EndY);
+
+									bool bInCell = false;
+									if (auto& Node = GetCell(GetGridX(X),Y))
+									{
+										TArray<FActorRepListType> ActorsInCell;
+										Node->GetAllActorsInNode_Debugging(ActorsInCell);
+										for (auto ActorInCell : ActorsInCell)
+										{
+											if (ActorInCell == DynamicActor)
+											{
+												if (bInCell)
+												{
+													UE_LOG(LogReplicationGraph, Warning, TEXT("  Actor is in cell multiple times! [%d, %d]"), X, Y);
+												}
+												bInCell = true;
+											}
+										}
+									}
+
+									if (bShouldBeInOld && bShouldBeInNew && bInCell)
+									{
+										// All good, didn't move
+										Str += "* ";
+									}
+									else if (!bShouldBeInOld && bShouldBeInNew && bInCell)
+									{
+										// All good, add
+										Str += "+ ";
+									}
+									else if (bShouldBeInOld && !bShouldBeInNew && !bInCell)
+									{
+										// All good, removed
+										Str += "- ";
+									}
+									else if (!bShouldBeInOld && !bShouldBeInNew && !bInCell)
+									{
+										// nada
+										Str += "  ";
+									}
+									else
+									{
+										UE_LOG(LogReplicationGraph, Warning, TEXT("  Bad update! Cell [%d,%d]. ShouldBeInOld: %d. ShouldBeInNew: %d. IsInCell: %d"), X, Y, bShouldBeInOld, bShouldBeInNew, bInCell);
+										Str += "! ";
+									}
+								}
+
+								UE_LOG(LogReplicationGraph, Display, TEXT("%s"), *Str);
+							}
+						}
+#endif
+
+						PreviousCellInfo = NewCellInfo;
+					}
+				}
+				else
+				{
+					// First time - Just add
+					GetGridNodesForActor(DynamicActor, NewCellInfo, GatheredNodes);
+					for (UReplicationGraphNode_GridCell* Node : GatheredNodes)
+					{
+						Node->AddDynamicActor(ActorInfo);
+					}
+
+					PreviousCellInfo = NewCellInfo;
+				}
+			}
+		}
+	}
+	
+	if (bNeedsRebuild)
+	{
+		RG_QUICK_SCOPE_CYCLE_COUNTER(UReplicationGraphNode_GridSpatialization2D_RebuildAll);
+
+		UE_LOG(LogReplicationGraph, Warning, TEXT("Rebuilding spatialization graph for bias %s"), *SpatialBias.ToString());
+		
+		// Tear down all existing nodes first. This marks them pending kill.
+		for (auto& InnerArray : Grid)
+		{
+			for (UReplicationGraphNode_GridCell*& N : InnerArray)
+			{
+				if (N)
+				{
+					N->TearDown();
+					N = nullptr;
+				}
+			}
+		}
+
+		// Force a garbage collection. Without this you may hit OOMs if rebuilding spatialization every frame for some period of time. 
+		// (Obviously not ideal to ever be doing this. But you are already hitching, might as well GC to avoid OOM crash).
+		
+		CollectGarbage( GARBAGE_COLLECTION_KEEPFLAGS, true );
+		
+		for (auto& MapIt : DynamicSpatializedActors)
+		{
+			FActorRepListType& DynamicActor = MapIt.Key;
+			if (ensureMsgf(IsActorValidForReplicationGather(DynamicActor), TEXT("%s not ready for replication."), *GetNameSafe(DynamicActor)))
+			{
+				FCachedDynamicActorInfo& DynamicActorInfo = MapIt.Value;
+				FActorCellInfo& PreviousCellInfo = DynamicActorInfo.CellInfo;
+				FNewReplicatedActorInfo& ActorInfo = DynamicActorInfo.ActorInfo;
+
+				const FVector Location3D = DynamicActor->GetActorLocation();
+				
+				FGlobalActorReplicationInfo& ActorRepInfo = GlobalRepMap->Get(DynamicActor);
+				ActorRepInfo.WorldLocation = Location3D;
+
+				const FActorCellInfo NewCellInfo = GetCellInfoForActor(DynamicActor, Location3D, ActorRepInfo.Settings.CullDistanceSquared);
+
+				GetGridNodesForActor(DynamicActor, NewCellInfo, GatheredNodes);
+				for (UReplicationGraphNode_GridCell* Node : GatheredNodes)
+				{
+					Node->AddDynamicActor(ActorInfo);
+				}
+
+				PreviousCellInfo = NewCellInfo;
+			}
+		}
+
+		for (auto& MapIt : StaticSpatializedActors)
+		{
+			FActorRepListType& StaticActor = MapIt.Key;
+			FCachedStaticActorInfo& StaticActorInfo = MapIt.Value;
+
+			if (ensureMsgf(IsActorValidForReplicationGather(StaticActor), TEXT("%s not ready for replication."), *GetNameSafe(StaticActor)))
+			{
+				PutStaticActorIntoCell(StaticActorInfo.ActorInfo, GlobalRepMap->Get(StaticActor), StaticActorInfo.bDormancyDriven);
+			}
+		}
+
+		bNeedsRebuild = false;
+	}
+}
+
+void UReplicationGraphNode_GridSpatialization2D::GatherActorListsForConnection(const FConnectionGatherActorListParameters& Params)
+{
+	if (Params.Viewer.ViewLocation.Z > ConnectionMaxZ)
+	{
+		return;
+	}
+
+	// Find out what bucket the view is in
+
+	int32 CellX = (Params.Viewer.ViewLocation.X - SpatialBias.X) / CellSize;
+	if (CellX < 0)
+	{
+		UE_LOG(LogReplicationGraph, Log, TEXT("Net view location.X %s is less than the spatial bias %s"), *Params.Viewer.ViewLocation.ToString(), *SpatialBias.ToString());
+		CellX = 0;
+	}
+	
+	TArray<UReplicationGraphNode_GridCell*>& GridX = GetGridX(CellX);
+
+	// -----------
+
+	int32 CellY = (Params.Viewer.ViewLocation.Y - SpatialBias.Y) / CellSize;
+	if (CellY < 0)
+	{
+		UE_LOG(LogReplicationGraph, Log, TEXT("Net view location.Y %s is less than the spatial bias %s"), *Params.Viewer.ViewLocation.ToString(), *SpatialBias.ToString());
+		CellY = 0;
+	}
+	if (GridX.Num() <= CellY)
+	{
+		GridX.SetNum(CellY+1);
+	}
+
+	if (UReplicationGraphNode_GridCell* Node = GridX[CellY])
+	{
+		Node->GatherActorListsForConnection(Params);
+	}
+}
+
+void UReplicationGraphNode_GridSpatialization2D::NotifyActorCullDistChange(AActor* Actor, FGlobalActorReplicationInfo& GlobalInfo, float OldDistSq)
+{
+	RG_QUICK_SCOPE_CYCLE_COUNTER(UReplicationGraphNode_GridSpatialization2D_NotifyActorCullDistChange);
+
+	// If this actor is statically spatialized then we need to remove it and readd it (this is a little wasteful but in practice not common/only happens at startup)
+	if (FCachedStaticActorInfo* StaticActorInfo = StaticSpatializedActors.Find(Actor))
+	{
+		// Remove with old distance
+		GetGridNodesForActor(Actor, GetCellInfoForActor(Actor, GlobalInfo.WorldLocation, OldDistSq), GatheredNodes);
+		for (UReplicationGraphNode_GridCell* Node : GatheredNodes)
+		{
+			Node->RemoveStaticActor(StaticActorInfo->ActorInfo, GlobalInfo, GlobalInfo.bWantsToBeDormant);
+		}
+
+		// Add new distances (there is some waste here but this hopefully doesn't happen much at runtime!)
+		GetGridNodesForActor(Actor, GetCellInfoForActor(Actor, GlobalInfo.WorldLocation, GlobalInfo.Settings.CullDistanceSquared), GatheredNodes);
+		for (UReplicationGraphNode_GridCell* Node : GatheredNodes)
+		{
+			Node->AddStaticActor(StaticActorInfo->ActorInfo, GlobalInfo, StaticActorInfo->bDormancyDriven);
+		}
+	}
+	else if (FCachedDynamicActorInfo* DynamicActorInfo = DynamicSpatializedActors.Find(Actor))
+	{
+		// Pull dynamic actor out of the grid. We will put him back on the next gather
+		
+		FActorCellInfo& PreviousCellInfo = DynamicActorInfo->CellInfo;
+		if (PreviousCellInfo.IsValid())
+		{
+			GetGridNodesForActor(Actor, PreviousCellInfo, GatheredNodes);
+			for (UReplicationGraphNode_GridCell* Node : GatheredNodes)
+			{
+				Node->RemoveDynamicActor(DynamicActorInfo->ActorInfo);
+			}
+			PreviousCellInfo.Reset();
+		}
+	}
+	else
+	{
+		UE_LOG(LogReplicationGraph, Warning, TEXT("UReplicationGraphNode_GridSpatialization2D::NotifyActorCullDistChange. %s Changed Cull Distance (%.2f -> %.2f) but is not in static or dynamic actor lists. %s"), *Actor->GetPathName(), FMath::Sqrt(OldDistSq), FMath::Sqrt(GlobalInfo.Settings.CullDistanceSquared), *GetPathName() );
+
+		// Search the entire grid. This is slow so only enabled if verify is on.
+		if (CVar_RepGraph_Verify)
+		{
+			bool bFound = false;
+			for (auto& InnerArray : Grid)
+			{
+				for (UReplicationGraphNode_GridCell* CellNode : InnerArray)
+				{
+					if (CellNode)
+					{
+						TArray<FActorRepListType> AllActors;
+						CellNode->GetAllActorsInNode_Debugging(AllActors);
+						if (AllActors.Contains(Actor))
+						{
+							UE_LOG(LogReplicationGraph, Warning, TEXT("  Its in node %s"), *CellNode->GetPathName());
+							bFound = true;
+						}
+					}
+				}
+			}
+			if (!bFound)
+			{
+				UE_LOG(LogReplicationGraph, Warning, TEXT("  Not in the grid at all!"));
+			}
+		}
+	}
+}
+
+
+// -------------------------------------------------------
+
+void UReplicationGraphNode_ClassCategories::NotifyAddNetworkActor(const FNewReplicatedActorInfo& ActorInfo)
+{
+	if (UReplicationGraphNode** Node = ActorClassMap.Get(ActorInfo.Class))
+	{
+		(*Node)->NotifyAddNetworkActor(ActorInfo);
+	}
+	else
+	{
+		UE_LOG(LogReplicationGraph, Warning, TEXT("NotifyAddNetworkActor called on unmapped actor class %s. (This actor shouldn't have been passed to this class)"), *GetNameSafe(ActorInfo.Class));
+	}
+}
+	
+void UReplicationGraphNode_ClassCategories::NotifyRemoveNetworkActor(const FNewReplicatedActorInfo& ActorInfo, bool bWarnIfNotFound)
+{
+	if (UReplicationGraphNode** Node = ActorClassMap.Get(ActorInfo.Class))
+	{
+		(*Node)->NotifyRemoveNetworkActor(ActorInfo, bWarnIfNotFound);
+	}
+	else
+	{
+		UE_LOG(LogReplicationGraph, Warning, TEXT("NotifyRemoveNetworkActor called on unmapped actor class %s. (This actor shouldn't have been passed to this class)"), *GetNameSafe(ActorInfo.Class));
+	}
+}
+
+void UReplicationGraphNode_ClassCategories::GatherActorListsForConnection(const FConnectionGatherActorListParameters& Params)
+{
+	for (const FCategoryMapping& Item : Categories)
+	{
+		if (Params.AttemptPushListCategory(Item.Category))
+		{
+			Item.Node->GatherActorListsForConnection(Params);
+			Params.PopListCategory();
+		}
+	}
+}
+
+void UReplicationGraphNode_ClassCategories::AddCategory(FReplicationListCategory Category, UClass* ActorClass)
+{
+	repCheckf(Category, TEXT("::AddCategory called with invalid frequency. Values should be >= 1."));
+	ActorClassMap.Set(ActorClass, GetCategoryMapping(Category).Node );
+}
+
+void UReplicationGraphNode_ClassCategories::AddCategory(FReplicationListCategory Category, const TArray<UClass*>& ActorClasses)
+{
+	UReplicationGraphNode* Node = GetCategoryMapping(Category).Node;
+	repCheck(Node);
+	for (UClass* Class : ActorClasses)
+	{
+		ActorClassMap.Set(Class, Node);
+	}
+}
+
+// -------------------------------------------------------
+
+void UReplicationGraphNode_FrequencyLimiter::NotifyAddNetworkActor(const FNewReplicatedActorInfo& ActorInfo)
+{
+	
+}
+	
+void UReplicationGraphNode_FrequencyLimiter::NotifyRemoveNetworkActor(const FNewReplicatedActorInfo& ActorInfo, bool bWarnIfNotFound)
+{
+	
+}
+
+void UReplicationGraphNode_FrequencyLimiter::GatherActorListsForConnection(const FConnectionGatherActorListParameters& Params)
+{
+	RG_QUICK_SCOPE_CYCLE_COUNTER(UReplicationGraphNode_FrequencyLimiter_GatherActorListsForConnection);
+}
+
+// -------------------------------------------------------
+
+UReplicationGraphNode_AlwaysRelevant::UReplicationGraphNode_AlwaysRelevant()
+{
+	bRequiresPrepareForReplicationCall = true;
+}
+
+void UReplicationGraphNode_AlwaysRelevant::PrepareForReplication()
+{
+	RG_QUICK_SCOPE_CYCLE_COUNTER(UReplicationGraphNode_AlwaysRelevant_PrepareForReplication);
+
+	if (ChildNode == nullptr)
+	{
+		ChildNode = CreateChildNode();
+	}
+
+	ChildNode->NotifyResetAllNetworkActors();
+	for (UClass* ActorClass : AlwaysRelevantClasses)
+	{
+		for (TActorIterator<AActor> It(GetWorld(), ActorClass); It; ++It)
+		{
+			AActor* Actor = *It;
+			if (IsActorValidForReplicationGather(Actor))
+			{			
+				ChildNode->NotifyAddNetworkActor( FNewReplicatedActorInfo(*It) );
+			}
+		}
+	}
+}
+
+void UReplicationGraphNode_AlwaysRelevant::AddAlwaysRelevantClass(UClass* Class)
+{
+	// Check that we aren't adding sub classes
+	for (UClass* ExistingClass : AlwaysRelevantClasses)
+	{
+		if (ExistingClass->IsChildOf(Class) || Class->IsChildOf(ExistingClass))
+		{
+			UE_LOG(LogReplicationGraph, Warning, TEXT("UReplicationGraphNode_AlwaysRelevant::AddAlwaysRelevantClass Adding class %s when %s is already in the list."), *Class->GetName(), *ExistingClass->GetName());
+		}
+	}
+
+
+	AlwaysRelevantClasses.AddUnique(Class);
+}
+
+void UReplicationGraphNode_AlwaysRelevant::GatherActorListsForConnection(const FConnectionGatherActorListParameters& Params)
+{
+	ChildNode->GatherActorListsForConnection(Params);
+}
+
+// -------------------------------------------------------
+		
+void UReplicationGraphNode_TearOff_ForConnection::GatherActorListsForConnection(const FConnectionGatherActorListParameters& Params)
+{
+	if (TearOffActors.Num() > 0)
+	{
+		ReplicationActorList.Reset();
+		FPerConnectionActorInfoMap& ActorInfoMap = Params.ConnectionManager.ActorInfoMap;
+
+		for (int32 idx=TearOffActors.Num()-1; idx >=0; --idx)
+		{
+			AActor* Actor = TearOffActors[idx].Actor;
+			const uint32 TearOffFrameNum = TearOffActors[idx].TearOffFrameNum;
+
+			// If actor is still valid (not pending kill etc)
+			if (Actor && IsActorValidForReplication(Actor))
+			{
+				// And has not replicated since becoming torn off
+				if (FConnectionReplicationActorInfo* ActorInfo = ActorInfoMap.Find(Actor))
+				{
+					if (ActorInfo->LastRepFrameNum <= TearOffFrameNum)
+					{
+						// Add it to the rep list
+						ReplicationActorList.Add(Actor);
+						continue;
+					}
+				}
+			}
+
+			// If we didn't get added to the list, remove this
+			TearOffActors.RemoveAtSwap(idx, 1, false);
+		}
+
+		if (ReplicationActorList.Num() > 0)
+		{
+			Params.OutGatheredReplicationLists.AddReplicationActorList(ReplicationActorList);
+		}
+	}
+}
+
+// -------------------------------------------------------
+
+void UReplicationGraphNode_AlwaysRelevant_ForConnection::GatherActorListsForConnection(const FConnectionGatherActorListParameters& Params)
+{
+	// This is not the ideal implementation. I would like to develop an API for making this a persistent list that is updated only when needed.
+	if (ChildNode == nullptr)
+	{
+		ChildNode = CreateChildNode();
+	}
+
+	ChildNode->NotifyResetAllNetworkActors();
+	ChildNode->NotifyAddNetworkActor( FNewReplicatedActorInfo(Params.Viewer.InViewer) );
+	ChildNode->NotifyAddNetworkActor( FNewReplicatedActorInfo(Params.Viewer.ViewTarget) );
+	ChildNode->GatherActorListsForConnection(Params);
+}
+
