Copying //UE4/Dev-Networking to //UE4/Dev-Main (Source: //UE4/Dev-Networking @ 4051526)

diff --git a/Engine/Plugins/Runtime/ReplicationGraph/Source/Public/ReplicationGraphTypes.h b/Engine/Plugins/Runtime/ReplicationGraph/Source/Public/ReplicationGraphTypes.h
new file mode 100644
--- /dev/null
+++ b/Engine/Plugins/Runtime/ReplicationGraph/Source/Public/ReplicationGraphTypes.h
@@ -0,0 +1,1231 @@
+// Copyright 1998-2018 Epic Games, Inc. All Rights Reserved.
+
+#pragma once
+
+#include "UObject/Package.h"
+#include "EngineGlobals.h"
+#include "Engine/LocalPlayer.h"
+#include "GameFramework/Actor.h"
+#include "UObject/ObjectKey.h"
+#include "UObject/UObjectHash.h"
+
+class AActor;
+class UNetConnection;
+class UNetReplicationGraphConnection;
+class UReplicationGraph;
+
+struct FActorDestructionInfo;
+
+DECLARE_LOG_CATEGORY_EXTERN( LogReplicationGraph, Log, All );
+
+// Check aliases for within the system. The intention is that these can be flipped to checkSlow once the system is stable.
+
+// Macro enable extra logging/bookkeeping. Disabled in Test/Shipping for perf.
+#if !(UE_BUILD_SHIPPING || UE_BUILD_TEST)
+	#define REPGRAPH_DETAILS 1
+	#define DO_REPGRAPH_DETAILS(X) (X)
+	#define repCheck(x) check(x)
+	#define repCheckf(expr, format, ...) checkf(expr, format, ##__VA_ARGS__ )
+	#define RG_QUICK_SCOPE_CYCLE_COUNTER(x) QUICK_SCOPE_CYCLE_COUNTER(x)
+	extern int32 CVar_RepGraph_Verify;
+#else 
+	#define REPGRAPH_DETAILS 0
+	#define DO_REPGRAPH_DETAILS(X) 0
+	#define repCheck(x)
+	#define repCheckf(expr, format, ...)
+	#define RG_QUICK_SCOPE_CYCLE_COUNTER(X)
+#endif
+
+
+#if !(UE_BUILD_SHIPPING || UE_BUILD_TEST)	
+#define REPGRAPH_STR(x) #x
+#define REPGRAPH_DEVCVAR_SHIPCONST(Type,VarName,Var,Value,Help) \
+	Type Var = Value; \
+	static FAutoConsoleVariableRef Var##CVar(TEXT(VarName), Var, TEXT(Help), ECVF_Cheat );
+#else
+#define REPGRAPH_DEVCVAR_SHIPCONST(Type,VarName,Var,Value,Help) \
+	const Type Var = Value;
+#endif
+
+
+
+// --------------------------------------------------------------------------------------------------------------------------------------------
+// --------------------------------------------------------------------------------------------------------------------------------------------
+// Replication List Category
+// --------------------------------------------------------------------------------------------------------------------------------------------
+// --------------------------------------------------------------------------------------------------------------------------------------------
+// 
+// ReplicationListCategory is just a bit in a uin32. All FReplicationListCategory must be power of 2 values so that they can be unioned in an FReplicationListCategoryMask.
+// This implementation was chosen to allow projects to be able to fully extend/customize categories while providing some type safety/clarity in code.
+//
+// The important high level concepts around categories are:
+//	1) They provide a way to broadly cull replication lists. Each list stored in the Replication Graph has a category associated with it. When gathering lists,
+//		the system will filter out categories that are not requested by the packet budget.
+//	2) Once we have the culled lists, categories are what the FPacketBudgets are specified in terms of.  That is, as we are filling up the packet, we using
+//		categories to determine how much space a list can occupy.
+
+
+// Enables run time ensure that FReplicationListCategory are initialized with power of 2 values
+#define REP_CAT_ENSURE_POW2 !(UE_BUILD_SHIPPING || UE_BUILD_TEST)
+
+struct REPLICATIONGRAPH_API FReplicationListCategory
+{
+	FReplicationListCategory() : Value(InvalidListCategory)	{ }
+	FReplicationListCategory(uint32 In): Value(In)
+	{ 
+#if REP_CAT_ENSURE_POW2
+		ensure(FMath::IsPowerOfTwo(Value));
+#endif
+	}
+	enum { InvalidListCategory = 0 };
+	operator uint32 () const { return Value; }
+	friend uint32 GetTypeHash(const FReplicationListCategory& ListCategory) { return GetTypeHash(ListCategory.Value); };
+	FORCEINLINE void Reset() { Value = InvalidListCategory; }
+	
+	FString GetDebugStringSlow() const;
+	
+	// Static methods for setting the string mapping function
+	typedef TFunction<FString (FReplicationListCategory)> FGetDebugStringFunc;
+	static void SetDebugStringFunc(FGetDebugStringFunc Func);
+	static void SetDebugStringFunc_StandardEnum(TCHAR* EnumString);
+
+protected:
+
+	uint32 Value;
+	static FGetDebugStringFunc CustomDebugStringFunc;
+
+	friend struct FReplicationListCategoryMask;
+};
+
+// Represents a mask or union of FReplicationListCategories. The purpose of this is to make clear what places are dealing with unions of categories vs which are expecting a single category
+struct REPLICATIONGRAPH_API FReplicationListCategoryMask : public FReplicationListCategory
+{
+	FReplicationListCategoryMask() : FReplicationListCategory() { }
+	FReplicationListCategoryMask(uint32 In) { Value = In; }
+
+	FReplicationListCategoryMask operator|(FReplicationListCategory const& Rhs) const { FReplicationListCategoryMask Result(*this); Result |= Rhs; return Result; }
+	FReplicationListCategoryMask operator&(FReplicationListCategory const& Rhs) const { FReplicationListCategoryMask Result(*this); Result &= Rhs; return Result; }
+
+	FReplicationListCategoryMask const& operator|=(FReplicationListCategory const& Rhs) { Value |= Rhs.Value; return *this; }
+	FReplicationListCategoryMask const& operator&=(FReplicationListCategory const& Rhs) { Value &= Rhs.Value; return *this; }
+};
+
+// --------------------------------------------------------------------------------------------------------------------------------------------
+// --------------------------------------------------------------------------------------------------------------------------------------------
+// Actor Replication List Types
+// --------------------------------------------------------------------------------------------------------------------------------------------
+// --------------------------------------------------------------------------------------------------------------------------------------------
+
+// Currently we are using raw AActor* for our replication lists. We may want to change this one day to be an ID or something else that can index
+// fast into arrays etc. (Currently we are using TMaps for associative data, static arrays would be faster but introduce constraints and headaches)
+// So for now, using a typedef and some helper functions to call out the interface/usage of FActorRepListType.
+
+typedef AActor* FActorRepListType;
+FORCEINLINE FString GetActorRepListTypeDebugString(const FActorRepListType& In) { return GetNameSafe(In); }
+FORCEINLINE UClass* GetActorRepListTypeClass(const FActorRepListType& In) { return In->GetClass(); }
+
+
+// Tests if an actor is valid for replication: not pending kill, etc. Says nothing about wanting to replicate or should replicate, etc.
+FORCEINLINE bool IsActorValidForReplication(const FActorRepListType& In) { return !In->IsPendingKill() && !In->IsPendingKillPending(); }
+
+// Tests if an actor is valid for replication gathering. Meaning, it can be gathered from the replication graph and considered for replication.
+FORCEINLINE bool IsActorValidForReplicationGather(const FActorRepListType& In)
+{ 
+	if (!IsActorValidForReplication(In))
+		return false;
+
+	if (In->GetIsReplicated() == false)
+		return false;
+
+	if (In->GetTearOff())
+		return false;
+
+	if (In->NetDormancy == DORM_Initial && In->IsNetStartupActor())
+		return false;
+
+/*
+	These checks were done in legacy code and we would like to avoid them.
+
+	// Actors should finish initialization outside of the replication loop. Maybe some weird multi frame delayed case?
+	if (!In->IsActorInitialized())
+		return false;
+
+	// This check is slow and is not needed unless you are streaming levels on the server. If needed this should be opt in globally some how.
+	ULevel* Level = In->GetLevel();
+	if ( Level->HasVisibilityChangeRequestPending() || Level->bIsAssociatingLevel )
+		return false;
+
+*/
+
+	return true;
+}
+
+/** The actual underlying list type that the system uses for a list of replicated actors. These are never manually allocated or even really used. All Replication code interacts with the "View" types below  */
+struct REPLICATIONGRAPH_API FActorRepList : FNoncopyable
+{
+	/** List Header */
+	int32 RefCount;
+	int32 Max;
+	int32 Num;
+
+	/** the "used bit" from the block we came from, so that we can clear it fast */
+	FBitReference UsedBitRef;
+
+	/** Variable length Data segment */
+	FActorRepListType Data[];
+
+	/** For TRefCountPtr usage */
+	void AddRef() { RefCount++; }
+	void Release();
+};
+
+/** This is a base templated type for the list "Views". This provides basic read only access between the two real views. */
+template<typename PointerType>
+struct TActorRepListViewBase
+{
+	TActorRepListViewBase() { }
+	TActorRepListViewBase(const PointerType& In) : RepList(In) { }
+
+	FORCEINLINE int32 Num() const { return RepList->Num; }
+	FORCEINLINE const FActorRepListType& operator[](int32 idx) const  { repCheck(RepList); repCheck(RepList->Max > idx); return RepList->Data[idx]; }
+
+	/** Resets the view to null - meaning it is not pointing to any list (as opposed to an empty list of some preallocated size) */
+	FORCEINLINE void ResetToNull() { RepList = nullptr; }
+	FORCEINLINE bool IsValid() const { return RepList != nullptr; }
+
+	int32 IndexOf(const FActorRepListType& Value) const
+	{
+		repCheck(RepList);
+		FActorRepListType* Data = RepList->Data;
+		int32 Num = RepList->Num;
+		for (int32 idx=0 ; idx < Num; ++idx)
+		{
+			if (Data[idx] == Value)
+			{
+				return idx;
+			}
+		}
+		return -1;
+	}
+
+	bool Contains(const FActorRepListType& Value) const { return IndexOf(Value) != -1; }
+
+	/** Add contents to TArray. this is intended for debugging/ease of use */
+	void AppendToTArray(TArray<FActorRepListType>& OutArray) const
+	{
+		if (IsValid())
+		{
+			for (FActorRepListType Actor : *this)
+			{
+				OutArray.Add(Actor);
+			}
+		}
+	}
+
+	void AppendToTSet(TSet<FActorRepListType>& OutSet) const
+	{
+		if (IsValid())
+		{
+			for (FActorRepListType Actor : *this)
+			{
+				OutSet.Add(Actor);
+			}
+		}
+	}
+	
+	PointerType RepList;
+
+private:
+
+	FORCEINLINE friend FActorRepListType* begin(const TActorRepListViewBase<PointerType>& View) { return View.RepList->Data; }
+	FORCEINLINE friend FActorRepListType* end(const TActorRepListViewBase<PointerType>& View) { return View.RepList->Data + View.RepList->Num; }
+};
+
+/** A view that maintains ownership/(ref counting) to an actor replication list.  */
+struct REPLICATIONGRAPH_API FActorRepListRefView : public TActorRepListViewBase<TRefCountPtr<FActorRepList>>
+{
+	/** Ideally, use Reset to set the initial size of the list. But if nothing is set, the first list we request will be of this size */
+	enum { InitialListSize = 4 } ;
+
+	FActorRepListRefView() { }
+	FActorRepListRefView(FActorRepList& InRepList) { RepList = &InRepList; }
+	FORCEINLINE FActorRepListType& operator[](int32 idx) const  { repCheck(RepList); repCheck(RepList->Max > idx); return RepList->Data[idx]; }
+
+	/** Initializes a new list for a given ExpectedMaxSize. Best practice is to call this once to get a good initial size to avoid reallocations/copying. Passing 0 will preserve the current size (if current size is also 0, InitialListSize is used) */
+	void Reset(int32 ExpectedMaxSize=0)
+	{
+		if (RepList && RepList->RefCount == 1 && RepList->Max >= ExpectedMaxSize)
+		{
+			// We can keep using this list
+			RepList->Num = 0;
+			CachedNum = 0;
+		}
+		else
+		{
+			// We must request a new list for this size
+			RequestNewList(ExpectedMaxSize > 0 ? ExpectedMaxSize : CachedNum, false);
+		}
+	}
+
+	/** Prepares the list for modifications. If this list is shared by other RefViews, then you will get a new underlying list to reference. ResetContent determines if the content is cleared or not (regardless of refcount/new list) */
+	void PrepareForWrite(bool bResetContent=false)
+	{
+		if (RepList == nullptr)
+		{
+			RequestNewList(InitialListSize, false);
+		}
+		else if (RepList->RefCount > 1)
+		{
+			// This list we are viewing is shared by others, so request a new one
+			RequestNewList(CachedNum, !bResetContent);
+		}
+		else if(bResetContent)
+		{
+			// We already have our own list but need to reset back to 0
+			RepList->Num = 0;
+			CachedNum = 0;
+		}
+	}
+
+	bool ConditionalAdd(const FActorRepListType& NewElement)
+	{
+		if (IsActorValidForReplicationGather(NewElement))
+		{
+			Add(NewElement);
+			return true;
+		}
+		return false;
+	}
+
+	void Add(const FActorRepListType& NewElement)
+	{
+		repCheckf(RepList != nullptr, TEXT("Invalid RepList when calling add new element to a list. Call ::PrepareForWrite or ::Reset before writing!"));
+		repCheckf(RepList->RefCount == 1, TEXT("Attempting to add new element to a list with >1 RefCount (%d). Call ::PrepareForWrite before writing!"), RepList->RefCount);
+		if (CachedNum == CachedMax)
+		{
+			// We can't add more to the list we are referencing, we need to get a new list and copy the contents over. This is transparent to the caller.
+			RequestNewList(CachedMax+1, true);
+		}
+		
+		CachedData[CachedNum++] = NewElement;
+		RepList->Num++;
+	}
+
+	bool Remove(const FActorRepListType& ElementToRemove)
+	{
+		int32 idx = IndexOf(ElementToRemove);
+		if (idx >= 0)
+		{
+			RemoveAtImpl(idx);
+			return true;
+		}
+		return false;
+	}
+
+	void RemoveAtSwap(int32 idx)
+	{
+		repCheck(RepList && Num() > idx);
+		CachedData[idx] = CachedData[CachedNum-1];
+		CachedNum--;
+		RepList->Num--;
+	}
+
+	void CopyContentsFrom(const FActorRepListRefView& Source);
+	bool VerifyContents_Slow() const;
+
+	int32 Num() const { return CachedNum; }
+
+private:
+
+	/** Cached data from our FActorRepList to avoid looking it up each time */
+	FActorRepListType* CachedData = nullptr;
+	int32 CachedNum  = 0;
+	int32 CachedMax = 0;
+
+	void RequestNewList(int32 NewSize, bool bCopyExistingContent);
+
+	void RemoveAtImpl(int32 Index)
+	{
+		repCheck(RepList && Num() > Index);
+		int32 NumToMove = CachedNum - Index - 1;
+		if (NumToMove)
+		{
+			FMemory::Memmove((void*)(&CachedData[Index]), (void*)(&CachedData[Index + 1]), NumToMove * sizeof(FActorRepListType));
+		}
+
+		CachedNum--;
+		RepList->Num--;
+	}
+};
+
+/** A read only, non owning (ref counting) view to an actor replication list: essentially a raw pointer and the category of the list. These are only created *from* FActorRepListRefView */
+struct REPLICATIONGRAPH_API FActorRepListRawView : public TActorRepListViewBase<FActorRepList*>
+{
+	/** Standard ctor: make raw view from ref view */
+	FActorRepListRawView(const FActorRepListRefView& Source, FReplicationListCategory InCategory) : Category(InCategory) { RepList = Source.RepList.GetReference(); }
+	
+	/** Sublist cstor: make raw view from portion of a ref view */
+	FActorRepListRawView(const FActorRepListRefView& Source, FReplicationListCategory InCategory, int32 StartIdx, int32 Num) : Category(InCategory)
+	{
+
+	}
+	
+	FActorRepListRefView ToRefView() const { return FActorRepListRefView(*RepList); }
+
+	FReplicationListCategory Category;
+};
+
+#if !(UE_BUILD_SHIPPING || UE_BUILD_TEST)
+	// Intended to be called from immediate mode window while debugging
+	extern "C" DLLEXPORT void PrintRepListDetails(int32 PoolSize, int32 BlockIdx, int32 ListIdx);
+	extern "C" DLLEXPORT void PrintRepListStats(int32 mode=0);
+	
+	void PrintRepListStatsAr(int32 mode, FOutputDevice& Ar=*GLog);	
+#endif
+
+/** To be called by projets to preallocate replication lists. This isn't strictly necessary: lists will be allocated on demand as well. */
+REPLICATIONGRAPH_API void PreAllocateRepList(int32 ListSize, int32 NumLists);
+
+
+// This represents "the list of gathered lists". This is what we push down the Replication Graph and nodes will either Push/Pop List Categories or will add their Replication Lists.
+struct REPLICATIONGRAPH_API FGatheredReplicationActorLists
+{
+	void AddReplicationActorList(const FActorRepListRefView& List)
+	{
+#if !(UE_BUILD_SHIPPING || UE_BUILD_TEST)
+		if (!ensureMsgf(CurrentListCategory, TEXT("FGatheredReplicationActorLists Attemping to add a ReplicationActorList with invalid ListCategory")))
+			return;
+
+		if (CVar_RepGraph_Verify)
+			List.VerifyContents_Slow();
+#endif
+		repCheck(List.IsValid());
+		if (List.Num() > 0)
+		{
+			OutReplicationLists.Emplace(FActorRepListRawView(List, CurrentListCategory)); 
+		}
+	}
+
+	void Reset() { OutReplicationLists.Reset(); ListCategoryStack.Reset(); }
+	
+	void PushListCategory(FReplicationListCategory ListCategory) { ListCategoryStack.Push(CurrentListCategory); CurrentListCategory = ListCategory; }
+	void PopListCategory() { CurrentListCategory = ListCategoryStack.Last(); ListCategoryStack.Pop(); }	
+
+	FORCEINLINE int32 Num() const { return OutReplicationLists.Num(); }
+	FORCEINLINE const FActorRepListRawView& operator[](int32 idx) const  { return OutReplicationLists[idx]; }
+	
+	FORCEINLINE friend FActorRepListRawView* begin(FGatheredReplicationActorLists& Array) { return Array.OutReplicationLists.GetData(); }
+	FORCEINLINE friend FActorRepListRawView* end(FGatheredReplicationActorLists& Array) { return Array.OutReplicationLists.GetData() + Array.OutReplicationLists.Num(); }
+
+private:	
+
+	TArray< FActorRepListRawView > OutReplicationLists;
+	FReplicationListCategory CurrentListCategory;
+	TArray< FReplicationListCategory, TInlineAllocator<1> > ListCategoryStack;
+};
+
+// --------------------------------------------------------------------------------------------------------------------------------------------
+// --------------------------------------------------------------------------------------------------------------------------------------------
+// Packet Budget
+// --------------------------------------------------------------------------------------------------------------------------------------------
+// --------------------------------------------------------------------------------------------------------------------------------------------
+
+// Packet Budget represents the desired makeup of a replication packet. It maps ListCategories to the desired amount of data, It is the const data that we look at while building packets.
+struct FPacketBudget
+{
+	/** What happens when a budget item ends while being under/over the target budget */
+	enum EBalancePolicy
+	{
+		/** The balance is absorbed by this item. */
+		Absorb,
+		/** The balance rolls over to the next item (e.g, if > 0, the next item gets more bits, if < 0, the next item gets less bits). Note that a budget item cannot be completely pushed out - it will replicate at least once! */
+		RollOver
+	};
+
+	struct FItem
+	{
+		FItem(FReplicationListCategoryMask InListCategory, int64 InMaxBits, EBalancePolicy InUnderBudgetPolicy, EBalancePolicy InOverBudgetPolicy) 
+			: ListCategory(InListCategory), MaxBits(InMaxBits), UnderBudgetPolicy(InUnderBudgetPolicy), OverBudgetPolicy(InOverBudgetPolicy)  { }
+
+		FReplicationListCategoryMask ListCategory;
+		int64 MaxBits = 0;
+		EBalancePolicy UnderBudgetPolicy;
+		EBalancePolicy OverBudgetPolicy;
+	};
+
+	/** The actual list of items that make up this packet budget */
+	TArray<FItem> BudgetItems;
+
+	/** Combined mask for all categories this packet budget cares about: used for broad filtering during gather phase */
+	FReplicationListCategoryMask ListCategoryMask;
+
+	/** Identifying string for debug/logging only */
+	const TCHAR* DebugName = nullptr;
+
+	/** Adds a budget item for a given category. */
+	void AddBudgetForCategory(FReplicationListCategoryMask Category, int64 MaxBits, EBalancePolicy UnderBudgetPolicy = EBalancePolicy::Absorb, EBalancePolicy OverBudgetPolicy = EBalancePolicy::Absorb)
+	{
+		BudgetItems.Emplace(FItem(Category, MaxBits, UnderBudgetPolicy, OverBudgetPolicy));
+	}
+
+	/** Finalizes the budget: call this (once) after finished calling AddBudgetForCategory (as many times as you need) */
+	void Finalize()
+	{
+		ListCategoryMask = FReplicationListCategoryMask();
+		for (auto& Item : BudgetItems)
+		{
+			ListCategoryMask |= Item.ListCategory;
+		}
+
+		BudgetItems.Sort([](const FItem& ItemA, const FItem& ItemB) { return ItemA.ListCategory < ItemB.ListCategory; } );
+	}
+};
+
+// Represents a filled budget item during an actual frame. These are produced by FPacketBudgetTracker and mainly used for logging/profiling etc.
+struct FPacketBudgetRecord
+{
+	struct FItem
+	{
+		int64 BitsWritten = 0;
+	};
+
+	void SetBudget(const FPacketBudget* InBudget) { Budget = InBudget; Items.Reset(); Items.SetNum(InBudget->BudgetItems.Num()); }
+	void TrackSentBits(int64 NumBits, int32 BudgetIdx) {	Items[BudgetIdx].BitsWritten += NumBits; }
+	void TrackActorChannelCloseBits(int64 NumBits) { ActorChannelCloseBits += NumBits; }
+	void TrackDestructionInfoBits(int64 NumBits) { DestructionInfoBits += NumBits; }
+	
+	const FPacketBudget* Budget = nullptr;
+	TArray<FItem, TInlineAllocator<8> >	Items;
+	int64 ActorChannelCloseBits = 0;
+	int64 DestructionInfoBits = 0;
+};
+
+// The tracker is what actually does the book keeping between a given FPacketBudget and the data that is pushed into packets. Essentially, the runtime/dynamic data as we are replicating actors.
+struct FPacketBudgetTracker
+{
+	FPacketBudgetTracker(const FPacketBudget& InBudget) : Budget(InBudget), CurrentIdx(0), RemainingBits(0) { }
+	
+	/** Sets the current index of the FPacketBudget::FItem that we are filling (this is an index into this->Budget.BudgetItems) */
+	void SetBudgetItemIndex(int32 Index)
+	{
+		if ((RemainingBits < 0 && Budget.BudgetItems[CurrentIdx].OverBudgetPolicy == FPacketBudget::RollOver) || (Budget.BudgetItems[CurrentIdx].UnderBudgetPolicy == FPacketBudget::RollOver))
+		{
+			// Rollover the balance (note, don't check RemainingBits > 0 because that is an extra banch and effectively the exact same)
+			RemainingBits += Budget.BudgetItems[Index].MaxBits;
+		}
+		else
+		{
+			// Absorb the balance
+			RemainingBits = Budget.BudgetItems[Index].MaxBits;
+		}
+
+		CurrentIdx = Index;
+	}
+	
+	/** Called as bits are written to the send buffer. Returns true if the current budget item is not exceeded. (note: technically they are not "sent" yet, not until the net driver's send buffer is flushed) */
+	bool TrackSentBits(int64 NumBits)
+	{
+		if (OutRecord)
+		{
+			OutRecord->TrackSentBits(NumBits, CurrentIdx);
+		}
+
+		RemainingBits -= NumBits;
+		return (RemainingBits > 0);
+	}
+
+	void TrackActorChannelCloseBits(int64 NumBits) { if (OutRecord) OutRecord->TrackActorChannelCloseBits(NumBits); }
+	void TrackDestructionInfoBits(int64 NumBits) { if (OutRecord) OutRecord->TrackDestructionInfoBits(NumBits); }
+
+	/** Optional output: useful for logging/debugging etc */
+	void SetOutputRecord(struct FPacketBudgetRecord* NewOutRecord) { OutRecord = NewOutRecord; }
+	
+private:
+
+	const FPacketBudget& Budget;
+	FReplicationListCategory CurrentListCategory;
+	int32 CurrentIdx;
+	int64 RemainingBits;
+
+	FPacketBudgetRecord* OutRecord = nullptr;
+};
+
+
+struct FPacketBudgetRecordBuffer
+{
+	const int32 BufferSize = 64;
+	TArray<FPacketBudgetRecord, TInlineAllocator<64> > Records;
+	uint32 NextWriteIdx = 0;
+
+	FPacketBudgetRecordBuffer()
+	{
+		Records.SetNum(BufferSize);
+	}
+
+	FPacketBudgetRecord* GetNextWriteRecord(const FPacketBudget& ForBudget)
+	{
+		FPacketBudgetRecord* Record = &Records[NextWriteIdx % BufferSize];
+		Record->SetBudget(&ForBudget);
+		NextWriteIdx++;
+		return Record;
+	}
+
+	int32 Num() const { return FMath::Min<uint32>(BufferSize, NextWriteIdx); }
+	FPacketBudgetRecord& GetAtIndex(int32 idx) { return Records[(NextWriteIdx - Num() + idx) % Records.Num()]; }
+	void Reset() { Records.Reset(); Records.SetNum(BufferSize); NextWriteIdx = 0; }
+};
+
+// --------------------------------------------------------------------------------------------------------------------------------------------
+// --------------------------------------------------------------------------------------------------------------------------------------------
+// Connection Gather Actor List Parameters
+// --------------------------------------------------------------------------------------------------------------------------------------------
+// --------------------------------------------------------------------------------------------------------------------------------------------
+
+// Parameter structure for what we actually pass down during the Gather phase.
+struct FConnectionGatherActorListParameters
+{
+	FConnectionGatherActorListParameters(FNetViewer& InViewer, UNetReplicationGraphConnection& InConnectionManager, TSet<FName>& InClientVisibleLevelNamesRef, FPacketBudget& InPacketBudget, uint32 InReplicationFrameNum, FGatheredReplicationActorLists& InOutGatheredReplicationLists)
+		: Viewer(InViewer), ConnectionManager(InConnectionManager), PacketBudget(InPacketBudget), ReplicationFrameNum(InReplicationFrameNum), OutGatheredReplicationLists(InOutGatheredReplicationLists), ClientVisibleLevelNamesRef(InClientVisibleLevelNamesRef)
+	{
+	}
+
+	/** In: The Data the nodes have to work with */
+	FNetViewer& Viewer;
+	UNetReplicationGraphConnection& ConnectionManager;
+	FPacketBudget& PacketBudget;
+	uint32 ReplicationFrameNum;
+
+	/** Out: The data nodes are going to add to */
+	FGatheredReplicationActorLists& OutGatheredReplicationLists;
+
+	/** Attempts to push the ListCategory, returns true if the list category is currently allowed. If this returns true, you are expected to call PopListCategory when you are done! */
+	bool AttemptPushListCategory(FReplicationListCategory ListCategory) const
+	{
+		if ((PacketBudget.ListCategoryMask & ListCategory))
+		{
+			OutGatheredReplicationLists.PushListCategory(ListCategory);
+			return true;
+		}
+		return false;
+	}
+
+	/** Pop the list category that was pushed */
+	void PopListCategory() const
+	{
+		OutGatheredReplicationLists.PopListCategory();
+	}
+
+	bool CheckClientVisibilityForLevel(const FName& StreamingLevelName) const
+	{
+		// We track the last checked StreamingLevelName that was visible. This is a simple optimization under the premise that in an ideal replication graph,
+		// a connection is only pulling from nodes with 1 conditional streaming level. (E.g, even if you map has 300 streaming levels, spatially the client should
+		// only be "over" one of those levels, and should only attempt to pull from nodes that have that one level in its conditional actor lists.
+		// 
+		// Really this is still left over from legacy implementation. I would rather completely skip nodes in the spatialization graph based on level streaming status,
+		// but the way that level streaming works, there is no mechanism or guarantees about what space a streaming level takes up. This is ok for now.
+		if (StreamingLevelName == LastCheckedVisibleLevelName)
+		{
+			return true;
+		}
+
+		const bool bVisible = ClientVisibleLevelNamesRef.Contains(StreamingLevelName);
+		if (bVisible)
+		{
+			LastCheckedVisibleLevelName = StreamingLevelName;
+		}
+		return bVisible;
+	}
+
+
+
+	// Cached off reference for fast Level Visibility lookup
+	TSet<FName>& ClientVisibleLevelNamesRef;
+
+private:
+
+	mutable FName LastCheckedVisibleLevelName;
+};
+
+// --------------------------------------------------------------------------------------------------------------------------------------------
+// --------------------------------------------------------------------------------------------------------------------------------------------
+// Per Class/Actor Global/PerConnectino Data Structs
+// --------------------------------------------------------------------------------------------------------------------------------------------
+// --------------------------------------------------------------------------------------------------------------------------------------------
+
+/** Per-Class actor data about how the actor replicates */
+struct FClassReplicationInfo
+{
+	FClassReplicationInfo() { }
+	float DistancePriorityScale = 1.f;
+	float StarvationPriorityScale = 1.f;
+	float CullDistanceSquared = 0.f;
+	
+	uint32 ReplicationPeriodFrame = 1;
+	uint32 ActorChannelFrameTimeout = 4;
+
+	FString BuildDebugStringDelta() const
+	{
+		FClassReplicationInfo DefaultValues;
+		FString Str;
+		if (CullDistanceSquared != DefaultValues.CullDistanceSquared)
+		{
+			Str += FString::Printf(TEXT("CullDistance: %.2f "), FMath::Sqrt(CullDistanceSquared));
+		}
+		if (StarvationPriorityScale != DefaultValues.StarvationPriorityScale)
+		{
+			Str += FString::Printf(TEXT("StarvationPriorityScale: %.2f "), StarvationPriorityScale);
+		}
+		if (DistancePriorityScale != DefaultValues.DistancePriorityScale)
+		{
+			Str += FString::Printf(TEXT("DistancePriorityScale: %.2f "), DistancePriorityScale);
+		}
+		if (ReplicationPeriodFrame != DefaultValues.ReplicationPeriodFrame)
+		{
+			Str += FString::Printf(TEXT("ReplicationPeriodFrame: %d "), ReplicationPeriodFrame);
+		}
+		if (ActorChannelFrameTimeout != DefaultValues.ActorChannelFrameTimeout)
+		{
+			Str += FString::Printf(TEXT("ActorChannelFrameTimeout: %d "), ActorChannelFrameTimeout);
+		}
+		return Str;
+	}
+};
+
+struct FGlobalActorReplicationInfo;
+
+DECLARE_MULTICAST_DELEGATE_FourParams(FNotifyActorChangeDormancy, FActorRepListType, FGlobalActorReplicationInfo&, ENetDormancy /*NewVlue*/, ENetDormancy /*OldValue*/);
+DECLARE_MULTICAST_DELEGATE_TwoParams(FNotifyActorFlushDormancy, FActorRepListType, FGlobalActorReplicationInfo&);
+DECLARE_MULTICAST_DELEGATE_TwoParams(FNotifyActorForceNetUpdate, FActorRepListType, FGlobalActorReplicationInfo&);
+
+struct FGlobalActorReplicationEvents
+{
+	FNotifyActorChangeDormancy	DormancyChange;
+	FNotifyActorFlushDormancy	DormancyFlush; // This delegate is cleared after broadcasting
+	FNotifyActorForceNetUpdate	ForceNetUpdate;
+};
+
+/** Per-Actor data that is global for the entire Replication Graph */
+struct FGlobalActorReplicationInfo
+{
+	FGlobalActorReplicationInfo(FClassReplicationInfo& ClassInfo) 
+		: LastPreReplicationFrame(0), WorldLocation(ForceInitToZero), Settings(ClassInfo) { }
+
+	// -----------------------------------------------------------
+	//	Dynamic state
+	// -----------------------------------------------------------
+
+	/** The last time AActor::PreReplication was called. Used to track when we need to call it again. */
+	uint32 LastPreReplicationFrame = 0;
+
+	/** the last time game code called ForceNetUpdate on this actor */
+	uint32 ForceNetUpdateFrame = 0;
+
+	/** Cached World Location of the actor */
+	FVector WorldLocation;
+
+	/** Mirrors AActor::NetDormany > DORM_Awake */
+	bool bWantsToBeDormant = false;
+
+	/** When this actor replicates, we replicate these actors immediately afterwards (they are not gathered/prioritized/etc) */
+	FActorRepListRefView DependentActorList;
+
+	// -----------------------------------------------------------
+	//	Class default mirrors: state that is initialized directly from class defaults (and can be later changed on a per-actor basis).
+	// -----------------------------------------------------------
+
+	FClassReplicationInfo Settings;
+
+	// -----------------------------------------------------------
+	//	Events
+	// -----------------------------------------------------------
+	FGlobalActorReplicationEvents Events;
+
+	void LogDebugString(FOutputDevice& Ar) const;
+};
+
+/** Templatd struct for mapping UClasses to some data type. The main things this provides is that if a UClass* was not explicitly added, it will climb the class heirachy and find the best match (and then store this for faster lookup next time) */
+template<typename ValueType>
+struct TClassMap
+{
+	/** Returns ClassInfo for a given class. */
+	ValueType& GetChecked(const UClass* Class)
+	{
+		ValueType* Ptr = Get(Class);
+		repCheckf(Ptr, TEXT("No ClassInfo found for %s"), *GetNameSafe(Class));
+		return *Ptr;
+	}
+
+	ValueType* Get(const UClass* Class)
+	{
+		FObjectKey ObjKey(Class);
+		if (ValueType* Ptr = Map.Find(ObjKey))
+		{
+			return Ptr;
+		}
+
+		// We haven't seen this class before, look it up (slower)
+		return GetClassInfoForNewClass_r(ObjKey, Class);
+	}
+
+	/** Returns if class has data in the map.  */
+	bool Contains(UClass* Class, bool bIncludeSuperClasses) const
+	{
+		if (bIncludeSuperClasses)
+		{
+			while(Class)
+			{
+				if (Map.Find(FObjectKey(Class)) != nullptr)
+				{
+					return true;
+				}
+
+				Class = Class->GetSuperClass();
+			}
+			return false;
+		}
+		
+		return (Map.Find(FObjectKey(Class)) != nullptr);
+	}
+
+	/** Sets class info for a given class. Call this in your Replication Graphs setup */
+	void Set(UClass* InClass, const ValueType& Value)
+	{
+		Map.Add(FObjectKey(InClass)) = Value;
+		
+		// Sets value for all derived classes. This is probably not useful since all classes may not be loaded anyways. TClassMap will 
+		// climb class heir achy when it encounters a new request. This shouldn't be too expensive, so the lazy approach seems better.
+		/*
+		TArray<UClass*> Classes = { InClass };
+		GetDerivedClasses(InClass, Classes, true);
+		for (UClass* Class : Classes)
+		{
+			UE_LOG(LogTemp, Display, TEXT("Adding for %s [From %s]"), *GetNameSafe(Class), *GetNameSafe(InClass));
+			Map.Add(FObjectKey(Class)) = Value;
+		}*/
+	}
+
+	FORCEINLINE typename TMap<FObjectKey, ValueType>::TIterator CreateIterator() { return Map.CreateIterator(); }
+
+private:
+
+	ValueType* GetClassInfoForNewClass_r(FObjectKey OriginalObjKey, const UClass* OriginalClass)
+	{
+		const UClass* Class = OriginalClass->GetSuperClass();
+		if (!Class)
+		{
+			return nullptr;
+		}
+
+		FObjectKey ObjKey(Class);
+
+		if (ValueType* Ptr = Map.Find(ObjKey))
+		{
+			// Set the original class's data to the super that we found (so that we don't have to climb the chain next time)
+			// Need to make a local copy in case the add reallocates the internal TMap data. This prevents Gil writing a blog post about you.
+			ValueType LocalData = *Ptr;
+			ValueType& NewData = Map.Add(OriginalObjKey);
+			NewData = LocalData;
+			return &NewData;
+		}
+
+		return GetClassInfoForNewClass_r(ObjKey, Class);
+	}
+
+	TMap<FObjectKey /** UClass */, ValueType> Map;
+};
+
+struct FGlobalActorReplicationInfoMap
+{
+	FGlobalActorReplicationInfoMap()
+	{
+		// Initialize so we always at least have AActor defined with default values. 
+		// It is ok to override this by calling SetClassInfo again.
+		ClassMap.Set(AActor::StaticClass(), FClassReplicationInfo());
+	}
+
+	/** Returns data associated with the actor. Will create it with default class values if necessary  */
+	FGlobalActorReplicationInfo& Get(const FActorRepListType& Actor)
+	{
+		// Quick lookup - this is the most common case
+		if (FGlobalActorReplicationInfo* Ptr = ActorMap.Find(Actor))
+		{
+			return *Ptr;
+		}
+
+		// We need to add data for this actor
+		FClassReplicationInfo& ClassInfo = GetClassInfo( GetActorRepListTypeClass(Actor) );
+		return ActorMap.Emplace(Actor, ClassInfo);
+	}
+
+	// Same as above but outputs bool if it was created. This is uncommonly called. Don't want to slow down the frequently called version
+	FGlobalActorReplicationInfo& Get(const FActorRepListType& Actor, bool& bWasCreated)
+	{
+		// Quick lookup - this is the most common case
+		if (FGlobalActorReplicationInfo* Ptr = ActorMap.Find(Actor))
+		{
+			return *Ptr;
+		}
+
+		bWasCreated = true;
+
+		// We need to add data for this actor
+		FClassReplicationInfo& ClassInfo = GetClassInfo( GetActorRepListTypeClass(Actor) );
+		return ActorMap.Emplace(Actor, ClassInfo);
+	}
+
+	/** Finds data associated with the actor but does not create if its not there yet. */
+	FORCEINLINE FGlobalActorReplicationInfo* Find(const FActorRepListType& Actor) { return ActorMap.Find(Actor); }
+
+	/** Removes actor data from map */
+	FORCEINLINE int32 Remove(const FActorRepListType& Actor) { return ActorMap.Remove(Actor); }
+
+	/** Returns ClassInfo for a given class. */
+	FORCEINLINE FClassReplicationInfo& GetClassInfo(UClass* Class) { return ClassMap.GetChecked(Class); }
+
+	/** Sets class info for a given class and its derived classes if desired. Call this in your Replication Graphs setup */
+	FORCEINLINE void SetClassInfo(UClass* InClass, const FClassReplicationInfo& Info) {	ClassMap.Set(InClass, Info); }
+	
+	FORCEINLINE TMap<FActorRepListType, FGlobalActorReplicationInfo>::TIterator CreateActorMapIterator() { return ActorMap.CreateIterator(); }
+	FORCEINLINE TMap<FObjectKey, FClassReplicationInfo>::TIterator CreateClassMapIterator() { return ClassMap.CreateIterator(); }
+
+private:
+
+	TMap<FActorRepListType, FGlobalActorReplicationInfo> ActorMap;
+	TClassMap<FClassReplicationInfo> ClassMap;
+};
+
+/** Per-Actor data that is stored per connection */
+struct FConnectionReplicationActorInfo
+{
+	FConnectionReplicationActorInfo() { }
+
+	FConnectionReplicationActorInfo(const FGlobalActorReplicationInfo& GlobalInfo)
+	{
+		// Pull data from the global actor info. This is done for things that we just want to duplicate in both places so that we can avoid a lookup into the global map
+		// and also for things that we want to be overridden per (connection/actor)
+
+		ReplicationPeriodFrame = GlobalInfo.Settings.ReplicationPeriodFrame;
+		CullDistanceSquared = GlobalInfo.Settings.CullDistanceSquared;
+	}
+
+	UActorChannel* Channel = nullptr;
+
+	float CullDistanceSquared = 0.f;
+
+	/** The next frame we are allowed to replicate on */
+	uint32 NextReplicationFrameNum = 0;
+
+	/** Min frames that have to pass between subsequent calls to ReplicateActor */
+	uint32 ReplicationPeriodFrame = 1;
+
+	/** The last frame that this actor replicated on to this connection */
+	uint32 LastRepFrameNum = 0;
+
+	/** The frame that this actor was (first) starved on. Meaning it wanted to replicate but we ran out of budget to do so. This is cleared everytime we do replicate */
+	uint32 StarvedFrameNum = 0;
+
+	/** The frame num that we will close the actor channel. This will get updated/pushed anytime the actor replicates based on FGlobalActorReplicationInfo::ActorChannelFrameTimeout  */
+	uint32 ActorChannelCloseFrameNum = 0;
+
+	bool bDormantOnConnection = false;
+	bool bTearOff = false;
+
+	void LogDebugString(FOutputDevice& Ar) const;
+};
+
+/** Map for Actor -> ConnectionActorInfo. This wraps the TMap mainly so we can do custom initialization in FindOrAdd. */
+struct FPerConnectionActorInfoMap
+{
+	FORCEINLINE_DEBUGGABLE FConnectionReplicationActorInfo& FindOrAdd(const FActorRepListType& Actor)
+	{
+		if (FConnectionReplicationActorInfo* Value = Map.Find(Actor))
+		{
+			return *Value;
+		}
+		FConnectionReplicationActorInfo NewInfo(GlobalMap->Get(Actor));
+
+		return Map.Emplace(Actor, GlobalMap->Get(Actor) );
+	}
+
+	FORCEINLINE FConnectionReplicationActorInfo* Find(const FActorRepListType& Actor)
+	{
+		return Map.Find(Actor);
+	}
+
+	FORCEINLINE TMap<FActorRepListType, FConnectionReplicationActorInfo>::TIterator CreateIterator()
+	{
+		return Map.CreateIterator();
+	}
+
+	FORCEINLINE FConnectionReplicationActorInfo& FindChecked(const FActorRepListType& Actor)
+	{
+		return Map.FindChecked(Actor);
+	}
+
+	FORCEINLINE void SetGlobalMap(FGlobalActorReplicationInfoMap* InGlobalMap)
+	{
+		GlobalMap = InGlobalMap;
+	}
+
+private:
+	TMap<FActorRepListType, FConnectionReplicationActorInfo> Map;
+	FGlobalActorReplicationInfoMap* GlobalMap;
+};
+
+/** Data that every replication graph has access to/is initialized with */
+struct FReplicationGraphGlobalData
+{
+	FReplicationGraphGlobalData() { }
+	FReplicationGraphGlobalData(FGlobalActorReplicationInfoMap* InRepMap, UWorld* InWorld) : GlobalActorReplicationInfoMap(InRepMap), World(InWorld) { }
+
+	FGlobalActorReplicationInfoMap* GlobalActorReplicationInfoMap = nullptr;
+
+	UWorld* World = nullptr;
+};
+
+// --------------------------------------------------------------------------------------------------------------------------------------------
+// --------------------------------------------------------------------------------------------------------------------------------------------
+// Prioritized Actor Lists
+// --------------------------------------------------------------------------------------------------------------------------------------------
+// --------------------------------------------------------------------------------------------------------------------------------------------
+
+/** Stores "full debug details" about how an actor was prioritized. This is not used in the actual replication code, just saved off for logging/debugging.  */
+struct FPrioritizedActorFullDebugDetails
+{
+	FPrioritizedActorFullDebugDetails(FActorRepListType InActor) : Actor(InActor) { }
+	bool operator==(const FActorRepListType& InActor) const { return Actor == InActor; }
+
+	FActorRepListType Actor;
+	float DistanceSq = 0.f;
+	float DistanceFactor = 0.f;
+
+	uint32 FramesSinceLastRap = 0;
+	float StarvationFactor = 0.f;
+
+	uint32 ForceNetUpdateDelta = 0;
+	float GameCodeScaling = 0.f;
+
+	FString BuildString() const
+	{
+		// This is pretty stupid slow but only needed for debug/logging
+		FString Str;
+		auto Append = [&Str](FString InStr) { Str += FString::Printf(TEXT("%-30s "), *InStr); };
+		
+		Append( DistanceFactor > 0.f ? FString::Printf(TEXT("(Dist: %.2f | %.2f) "), FMath::Sqrt(DistanceSq), DistanceFactor) : TEXT(""));
+		Append( StarvationFactor > 0.f ? FString::Printf(TEXT("(FramesLastRep: %d | %.2f) "), FramesSinceLastRap, StarvationFactor) : TEXT("") );
+		Append ( GameCodeScaling > 0.f ? FString::Printf(TEXT("(ForceNetUpdate: %d | %.2f) "), ForceNetUpdateDelta, GameCodeScaling) : TEXT(""));
+		
+		return Str;
+	}
+};
+
+/** Debug data about an actor that was skipped during the prioritization phase */
+struct FSkippedActorFullDebugDetails
+{
+	FSkippedActorFullDebugDetails(FActorRepListType InActor) : Actor(InActor) { }
+	FActorRepListType Actor;
+	bool bWasDormant = false; // If set, was skipped because it is dormant on this connection
+	float DistanceCulled = 0.f; // If set, was skipped due to distance culling
+	uint32 FramesTillNextReplication = 0; // If set, was skipped due to not being ready for replication
+};
+
+/** Prioritized List of actors to replicate. This is what we actually use to replicate actors. */
+struct FPrioritizedRepList
+{
+	FPrioritizedRepList(FReplicationListCategory InCategory) : ListCategory(InCategory) { }
+	struct FItem
+	{
+		// FIXME: We cannot store pointers to backing actor data yet because of how they are stored in (unstable) TMaps
+
+		FItem(float InPriority, FActorRepListType InActor /*, FGlobalActorReplicationInfo* InGlobal, FReplicationConnectionActorInfo* InConn*/) 
+			: Priority(InPriority), Actor(InActor) /*, GlobalData(InGlobal), ConnectionData(InConn)*/ { }
+		bool operator<(const FItem& Other) const { return Priority < Other.Priority; }
+
+		float Priority;
+		FActorRepListType Actor;
+
+		/*
+		FGlobalActorReplicationInfo* GlobalData;
+		FReplicationConnectionActorInfo* ConnectionData;
+		*/
+	};
+
+	TArray<FItem> Items;
+	FReplicationListCategory ListCategory;
+	
+#if REPGRAPH_DETAILS
+	FPrioritizedActorFullDebugDetails* GetNextFullDebugDetails(FActorRepListType Actor)
+	{
+		if (FullDebugDetails.IsValid() == false)
+		{
+			FullDebugDetails = MakeUnique<TArray<FPrioritizedActorFullDebugDetails> >();
+		}
+		return new (*FullDebugDetails) FPrioritizedActorFullDebugDetails(Actor);
+	}
+	TUniquePtr<TArray<FPrioritizedActorFullDebugDetails> > FullDebugDetails;
+	
+	
+	FSkippedActorFullDebugDetails* GetNextSkippedDebugDetails(FActorRepListType Actor)
+	{
+		if (SkippedDebugDetails.IsValid() == false)
+		{
+			SkippedDebugDetails = MakeUnique<TArray<FSkippedActorFullDebugDetails> >();
+		}
+		return new (*SkippedDebugDetails) FSkippedActorFullDebugDetails(Actor);
+	}
+	TUniquePtr<TArray<FSkippedActorFullDebugDetails> > SkippedDebugDetails;
+#endif
+};
+
+
+// --------------------------------------------------------------------------------------------------------------------------------------------
+// --------------------------------------------------------------------------------------------------------------------------------------------
+// Prioritized Actor Lists
+// --------------------------------------------------------------------------------------------------------------------------------------------
+// --------------------------------------------------------------------------------------------------------------------------------------------
+
+/** This is the struct we use to push new replication actors into the graph. "New" doesn't mean "newly spawned" it means "new to the graph". FIXME: Please suggest a better name! */
+struct FNewReplicatedActorInfo
+{
+	explicit FNewReplicatedActorInfo(const FActorRepListType& InActor) : Actor(InActor), Class(InActor->GetClass())
+	{
+		ULevel* Level = Cast<ULevel>(GetActor()->GetOuter());
+		if (Level && Level->IsPersistentLevel() == false)
+		{
+			StreamingLevelName = Level->GetOutermost()->GetFName();
+		}
+	}
+
+	AActor* GetActor() const { return Actor; }
+
+	FActorRepListType Actor;
+	FName StreamingLevelName;
+	UClass* Class;
+};
+
+// --------------------------------------------------------------------------------------------------------------------------------------------
+// --------------------------------------------------------------------------------------------------------------------------------------------
+// RPCs
+// --------------------------------------------------------------------------------------------------------------------------------------------
+// --------------------------------------------------------------------------------------------------------------------------------------------
+
+struct FRPCSendPolicyInfo
+{
+	FRPCSendPolicyInfo(const bool bInSendImmediately) : bSendImmediately(bInSendImmediately) { }
+	bool bSendImmediately;
+
+	// Suspect that this will grow over time. Possibly things like "min distance to send immediately" etc
+};
+
+// --------------------------------------------------------------------------------------------------------------------------------------------
+// --------------------------------------------------------------------------------------------------------------------------------------------
+// Debug Info
+// --------------------------------------------------------------------------------------------------------------------------------------------
+// --------------------------------------------------------------------------------------------------------------------------------------------
+
+struct FReplicationGraphDebugInfo
+{
+	FReplicationGraphDebugInfo( FOutputDevice& InAr ) : Ar(InAr) { }
+
+	enum EFlags
+	{
+		ShowActors,
+		ShowClasses,
+		ShowNativeClasses,
+		ShowTotalCount
+	};
+	
+	FOutputDevice& Ar;
+	EFlags Flags;
+	FString CurrentIndentString;
+	const FString IndentString = TEXT("  ");
+
+	void Log(const FString& Str) { Ar.Logf(TEXT("%s%s"), *CurrentIndentString, *Str); }
+
+	void PushIndent() { CurrentIndentString += IndentString; }
+	void PopIndent() { CurrentIndentString = CurrentIndentString.LeftChop(IndentString.Len()); }
+};
+
+REPLICATIONGRAPH_API void LogActorRepList(FReplicationGraphDebugInfo& DebugInfo, FString Prefix, const FActorRepListRefView& List);
+
+struct FPrioritizedActorDebugInfo
+{
+	int32 DormantCount = 0;
+	int32 NotReadyCount = 0;
+	int32 DistanceCulledCount = 0;
+	int32 ReplicatedCount = 0;
+	int32 StarvedCount = 0;
+};
+
+
+struct FNativeClassAccumulator
+{
+	void Increment(UClass* Class)
+	{
+		while (Class)
+		{
+			if (Class->IsNative())
+			{
+				break;
+			}
+			Class = Class->GetSuperClass();
+		}
+		
+		Map.FindOrAdd(Class)++;
+	}
+
+	FString BuildString()
+	{
+		FString Str;
+		Sort();
+		for (auto& It: Map)
+		{
+			Str += FString::Printf(TEXT("[%s, %d] "), *It.Key->GetName(), It.Value);
+			
+		}
+		return Str;
+	}
+
+	void Reset() { Map.Reset(); }
+	void Sort() { Map.ValueSort(TGreater<int32>()); }
+	TMap<UClass*, int32> Map;
+};
+
+// --------------------------------------------------------------------------------------------------------------------------------------------
+// --------------------------------------------------------------------------------------------------------------------------------------------
+// Profiling
+// --------------------------------------------------------------------------------------------------------------------------------------------
+// --------------------------------------------------------------------------------------------------------------------------------------------
+
+#define USE_REPCSVPROFILER !UE_BUILD_SHIPPING
+#define USE_REPCSVPROFILER_EX !(UE_BUILD_SHIPPING || UE_BUILD_TEST)
+
+#if USE_REPCSVPROFILER
+struct FReplicationGraphProfiler
+{
+	FReplicationGraphProfiler()
+	{
+		bEnabled = FParse::Param(FCommandLine::Get(), TEXT("RepGraphProfile"));
+	}
+
+	~FReplicationGraphProfiler()
+	{
+		End();
+	}
+
+	void OnClientConnect();
+	void StartRepFrame();
+	void EndRepFrame();
+
+private:
+	void End();
+	
+	double StartTime = 0.f;
+
+	bool bEnabled = false;
+	bool bStarted = false;
+	uint32 KillFrame = -1;
+	float TimeLimit = 60.f * 10;
+};
+#endif
+
diff --git a/Engine/Plugins/Runtime/ReplicationGraph/Source/ReplicationGraph.Build.cs b/Engine/Plugins/Runtime/ReplicationGraph/Source/ReplicationGraph.Build.cs
new file mode 100644
--- /dev/null
+++ b/Engine/Plugins/Runtime/ReplicationGraph/Source/ReplicationGraph.Build.cs
@@ -0,0 +1,25 @@
+// Copyright 1998-2017 Epic Games, Inc. All Rights Reserved.
+
+using UnrealBuildTool;
+
+public class ReplicationGraph : ModuleRules
+{
+	public ReplicationGraph(ReadOnlyTargetRules Target) : base(Target)
+	{
+		PublicIncludePaths.Add(ModuleDirectory + "/Public");
+
+		PrivateDependencyModuleNames.AddRange(
+			new string[] {
+				"Core",
+				"CoreUObject",
+				"Engine",
+				"EngineSettings"
+			}
+		);
+
+        PrivateIncludePaths.AddRange(
+            new string[] {
+                "Private"
+            });     
+    }
+}
diff --git a/Engine/Source/Editor/Matinee/Private/MatineeMenus.cpp b/Engine/Source/Editor/Matinee/Private/MatineeMenus.cpp
--- a/Engine/Source/Editor/Matinee/Private/MatineeMenus.cpp
+++ b/Engine/Source/Editor/Matinee/Private/MatineeMenus.cpp
@@ -4676,8 +4676,8 @@ void FMatinee::NormalizeVelocity()
 			SegmentEndTime = Dialog->IntervalEnd;
 			
 			// Make sure the user didnt enter any invalid values
-			FMath::Clamp(SegmentStartTime, FullStartTime, FullEndTime );
-			FMath::Clamp(SegmentEndTime, FullStartTime, FullEndTime );
+			SegmentStartTime = FMath::Clamp(SegmentStartTime, FullStartTime, FullEndTime );
+			SegmentEndTime = FMath::Clamp(SegmentEndTime, FullStartTime, FullEndTime );
 
 			// If we have a valid start and end time, normalize the track
 			if( SegmentStartTime != SegmentEndTime )
diff --git a/Engine/Source/Editor/UnrealEd/Private/EditorEngine.cpp b/Engine/Source/Editor/UnrealEd/Private/EditorEngine.cpp
--- a/Engine/Source/Editor/UnrealEd/Private/EditorEngine.cpp
+++ b/Engine/Source/Editor/UnrealEd/Private/EditorEngine.cpp
@@ -873,7 +873,7 @@ void UEditorEngine::Init(IEngineLoop* InEngineLoop)
 
 	FAssetRegistryModule& AssetRegistryModule = FModuleManager::LoadModuleChecked<FAssetRegistryModule>("AssetRegistry");
 	AssetRegistryModule.Get().OnInMemoryAssetCreated().AddUObject(this, &UEditorEngine::OnAssetCreated);
-	
+
 	FEditorDelegates::BeginPIE.AddLambda([](bool)
 	{
 		FTextLocalizationManager::Get().PushAutoEnableGameLocalizationPreview();
@@ -883,6 +883,14 @@ void UEditorEngine::Init(IEngineLoop* InEngineLoop)
 		GEngine->GetDynamicResolutionState()->ResetHistory();
 	});
 
+	FEditorDelegates::PrePIEEnded.AddLambda([this](bool)
+	{
+		if (GetPIEWorldContext() != nullptr && GetPIEWorldContext()->World() != nullptr)
+		{
+			GetPIEWorldContext()->World()->DestroyDemoNetDriver();
+		}
+	});
+
 	FEditorDelegates::EndPIE.AddLambda([](bool)
 	{
 		FTextLocalizationManager::Get().PopAutoEnableGameLocalizationPreview();
diff --git a/Engine/Source/Programs/AutomationTool/Scripts/BuildPhysX.Automation.cs b/Engine/Source/Programs/AutomationTool/Scripts/BuildPhysX.Automation.cs
--- a/Engine/Source/Programs/AutomationTool/Scripts/BuildPhysX.Automation.cs
+++ b/Engine/Source/Programs/AutomationTool/Scripts/BuildPhysX.Automation.cs
@@ -800,7 +800,7 @@ class BuildPhysX : BuildCommand
 
 		return NewPath;
 	}
-	private static void SetupBuildEnvironment()
+	private static void SetupStaticBuildEnvironment()
 	{
 		if (!Utils.IsRunningOnMono)
 		{
@@ -827,14 +827,19 @@ class BuildPhysX : BuildCommand
 				Environment.SetEnvironmentVariable("PATH", CMakePath + ";" + MakePath + ";" + Environment.GetEnvironmentVariable("PATH"));
 				Log("set {0}={1}", "PATH", Environment.GetEnvironmentVariable("PATH"));
 			}
-			// ================================================================================
-			// HTML5
-			// FIXME: only run this if GetTargetPlatforms() contains HTML5
+		}
+	}
 
+	private void SetupInstanceBuildEnvironment()
+	{
+		// ================================================================================
+		// HTML5
+		if (GetTargetPlatforms().Any(X => X.Platform == UnrealTargetPlatform.HTML5))
+		{
 			// override BuildConfiguration defaults - so we can use HTML5SDKInfo
 			string EngineSourceDir = GetProjectDirectory(PhysXTargetLib.PhysX, new TargetPlatformData(UnrealTargetPlatform.HTML5)).ToString();
-			EngineSourceDir = Regex.Replace(EngineSourceDir, @"\\" , "/");
-			EngineSourceDir = Regex.Replace(EngineSourceDir, ".*Engine/" , "");
+			EngineSourceDir = Regex.Replace(EngineSourceDir, @"\\", "/");
+			EngineSourceDir = Regex.Replace(EngineSourceDir, ".*Engine/", "");
 
 			if (!HTML5SDKInfo.IsSDKInstalled())
 			{
@@ -1462,7 +1467,8 @@ class BuildPhysX : BuildCommand
 
 	public override void ExecuteBuild()
 	{
-		SetupBuildEnvironment();
+		SetupStaticBuildEnvironment();
+		SetupInstanceBuildEnvironment();
 
 		bool bBuildSolutions = true;
 		if (ParseParam("SkipBuildSolutions"))
diff --git a/Engine/Source/Runtime/Core/Private/ProfilingDebugging/CsvProfiler.cpp b/Engine/Source/Runtime/Core/Private/ProfilingDebugging/CsvProfiler.cpp
--- a/Engine/Source/Runtime/Core/Private/ProfilingDebugging/CsvProfiler.cpp
+++ b/Engine/Source/Runtime/Core/Private/ProfilingDebugging/CsvProfiler.cpp
@@ -2070,6 +2070,17 @@ int32 FCsvProfiler::GetCaptureFrameNumber()
 	return CaptureFrameNumber;
 }
 
+bool FCsvProfiler::EnableCategoryByString(const FString& CategoryName) const
+{
+	int32 Category = FCsvCategoryData::Get()->GetCategoryIndex(CategoryName);
+	if (Category >= 0)
+	{
+		GCsvCategoriesEnabled[Category] = true;
+		return true;
+	}
+	return false;
+}
+
 bool FCsvProfiler::IsCapturing_Renderthread()
 {
 	check(IsInRenderingThread());
diff --git a/Engine/Source/Runtime/Core/Public/ProfilingDebugging/CsvProfiler.h b/Engine/Source/Runtime/Core/Public/ProfilingDebugging/CsvProfiler.h
--- a/Engine/Source/Runtime/Core/Public/ProfilingDebugging/CsvProfiler.h
+++ b/Engine/Source/Runtime/Core/Public/ProfilingDebugging/CsvProfiler.h
@@ -174,6 +174,8 @@ public:
 
 	CORE_API int32 GetCaptureFrameNumber();
 
+	CORE_API bool EnableCategoryByString(const FString& CategoryName) const;
+
 	/** Per-frame update */
 	CORE_API void BeginFrame();
 	CORE_API void EndFrame();
diff --git a/Engine/Source/Runtime/Engine/Classes/Engine/ActorChannel.h b/Engine/Source/Runtime/Engine/Classes/Engine/ActorChannel.h
--- a/Engine/Source/Runtime/Engine/Classes/Engine/ActorChannel.h
+++ b/Engine/Source/Runtime/Engine/Classes/Engine/ActorChannel.h
@@ -121,7 +121,7 @@ public:
 	bool ProcessQueuedBunches();
 
 	virtual void ReceivedNak( int32 NakPacketId ) override;
-	virtual void Close() override;
+	virtual int64 Close() override;
 	virtual FString Describe() override;
 
 public:
@@ -129,15 +129,15 @@ public:
 	/** UActorChannel interface and accessors. */
 	AActor* GetActor() {return Actor;}
 
-	/** Replicate this channel's actor differences. */
-	bool ReplicateActor();
+	/** Replicate this channel's actor differences. Returns how many bits were replicated (does not include non-bunch packet overhead) */
+	int64 ReplicateActor();
 
 	/** Allocate replication tables for the actor channel. */
 	void SetChannelActor( AActor* InActor );
 
 	virtual void NotifyActorChannelOpen(AActor* InActor, FInBunch& InBunch);
 
-	void SetChannelActorForDestroy( struct FActorDestructionInfo *DestructInfo );
+	int64 SetChannelActorForDestroy( struct FActorDestructionInfo *DestructInfo );
 
 	/** Append any export bunches */
 	virtual void AppendExportBunches( TArray< FOutBunch* >& OutExportBunches ) override;
diff --git a/Engine/Source/Runtime/Engine/Classes/Engine/Canvas.h b/Engine/Source/Runtime/Engine/Classes/Engine/Canvas.h
--- a/Engine/Source/Runtime/Engine/Classes/Engine/Canvas.h
+++ b/Engine/Source/Runtime/Engine/Classes/Engine/Canvas.h
@@ -566,7 +566,7 @@ public:
 	 * @param Thickness					How many pixels thick the box lines should be.
 	 */
 	UFUNCTION(BlueprintCallable, Category=Canvas, meta=(DisplayName="Draw Box", ScriptName="DrawBox"))
-	void K2_DrawBox(FVector2D ScreenPosition, FVector2D ScreenSize, float Thickness=1.0f);
+	void K2_DrawBox(FVector2D ScreenPosition, FVector2D ScreenSize, float Thickness=1.0f, FLinearColor RenderColor=FLinearColor::White);
 
 	/**
 	 * Draws a set of triangles on the Canvas.
diff --git a/Engine/Source/Runtime/Engine/Classes/Engine/Channel.h b/Engine/Source/Runtime/Engine/Classes/Engine/Channel.h
--- a/Engine/Source/Runtime/Engine/Classes/Engine/Channel.h
+++ b/Engine/Source/Runtime/Engine/Classes/Engine/Channel.h
@@ -82,8 +82,8 @@ public:
 	/** Set the closing flag. */
 	virtual void SetClosingFlag();
 
-	/** Close the base channel. */
-	virtual void Close();
+	/** Close the base channel. Returns how many bits were written to the send buffer */
+	virtual int64 Close();
 
 	/** Describe the channel. */
 	virtual FString Describe();
diff --git a/Engine/Source/Runtime/Engine/Classes/Engine/DemoNetConnection.h b/Engine/Source/Runtime/Engine/Classes/Engine/DemoNetConnection.h
--- a/Engine/Source/Runtime/Engine/Classes/Engine/DemoNetConnection.h
+++ b/Engine/Source/Runtime/Engine/Classes/Engine/DemoNetConnection.h
@@ -19,10 +19,14 @@ struct FQueuedDemoPacket
 	/** The size of the packet in bits */
 	int32 SizeBits;
 
+	/** Index of the level this packet is associated with. 0 indicates no association. */
+	uint32 SeenLevelIndex;
+
 public:
 	FORCEINLINE FQueuedDemoPacket(uint8* InData, int32 InSizeBytes, int32 InSizeBits)
 		: Data()
 		, SizeBits(InSizeBits)
+		, SeenLevelIndex(0)
 	{
 		Data.AddUninitialized(InSizeBytes);
 		FMemory::Memcpy(Data.GetData(), InData, InSizeBytes);
diff --git a/Engine/Source/Runtime/Engine/Classes/Engine/DemoNetDriver.h b/Engine/Source/Runtime/Engine/Classes/Engine/DemoNetDriver.h
--- a/Engine/Source/Runtime/Engine/Classes/Engine/DemoNetDriver.h
+++ b/Engine/Source/Runtime/Engine/Classes/Engine/DemoNetDriver.h
@@ -15,11 +15,14 @@
 #include "Misc/NetworkVersion.h"
 #include "NetworkReplayStreaming.h"
 #include "Engine/DemoNetConnection.h"
+#include "Net/RepLayout.h"
 #include "DemoNetDriver.generated.h"
 
 class Error;
 class FNetworkNotify;
 class UDemoNetDriver;
+class UDemoNetConnection;
+class FRepState;
 
 DECLARE_LOG_CATEGORY_EXTERN( LogDemo, Log, All );
 
@@ -28,9 +31,6 @@ DECLARE_DELEGATE_OneParam(FOnGotoTimeDelegate, const bool /* bWasSuccessful */);
 
 DECLARE_MULTICAST_DELEGATE(FOnDemoFinishPlaybackDelegate);
 
-class UDemoNetDriver;
-class UDemoNetConnection;
-
 class FQueuedReplayTask
 {
 public:
@@ -69,9 +69,10 @@ typedef TIndirectArray< FReplayExternalData > FReplayExternalDataArray;
 
 struct FPlaybackPacket
 {
-	TArray< uint8 > Data;
-	float			TimeSeconds;
-	int32			LevelIndex;
+	TArray< uint8 >		Data;
+	float				TimeSeconds;
+	int32				LevelIndex;
+	uint32				SeenLevelIndex;
 };
 
 enum ENetworkVersionHistory
@@ -84,13 +85,18 @@ enum ENetworkVersionHistory
 	HISTORY_MULTIPLE_LEVELS					= 6,			// Replays support seamless travel between levels
 	HISTORY_MULTIPLE_LEVELS_TIME_CHANGES	= 7,			// Save out the time that level changes happen
 	HISTORY_DELETED_STARTUP_ACTORS			= 8,			// Save DeletedNetStartupActors inside checkpoints
-	HISTORY_HEADER_FLAGS					= 9				// Save out enum flags with demo header
+	HISTORY_HEADER_FLAGS					= 9,			// Save out enum flags with demo header
+	HISTORY_LEVEL_STREAMING_FIXES			= 10,			// Optional level streaming fixes.
+	
+	// -----<new versions can be added before this line>-------------------------------------------------
+	HISTORY_PLUS_ONE,
+	HISTORY_LATEST 							= HISTORY_PLUS_ONE - 1
 };
 
 static const uint32 MIN_SUPPORTED_VERSION = HISTORY_EXTRA_VERSION;
 
 static const uint32 NETWORK_DEMO_MAGIC				= 0x2CF5A13D;
-static const uint32 NETWORK_DEMO_VERSION			= HISTORY_HEADER_FLAGS;
+static const uint32 NETWORK_DEMO_VERSION			= HISTORY_LATEST;
 static const uint32 MIN_NETWORK_DEMO_VERSION		= HISTORY_EXTRA_VERSION;
 
 static const uint32 NETWORK_DEMO_METADATA_MAGIC		= 0x3D06B24E;
@@ -128,9 +134,10 @@ enum class EReplayHeaderFlags : uint32
 {
 	None				= 0,
 	ClientRecorded		= ( 1 << 0 ),
+	HasStreamingFixes	= ( 1 << 1 ),
 };
 
-ENUM_CLASS_FLAGS(EReplayHeaderFlags)
+ENUM_CLASS_FLAGS(EReplayHeaderFlags);
 
 struct FNetworkDemoHeader
 {
@@ -227,8 +234,21 @@ struct FRollbackNetStartupActorInfo
 	FRotator	Rotation;
 	UPROPERTY()
 	ULevel*		Level;
+
+	TSharedPtr<FRepState> RepState;
+	UPROPERTY()
+	TArray<UObject*> ObjReferences;
+};
+
+struct FDemoSavedRepObjectState
+{
+	TWeakObjectPtr<const UObject> Object;
+	TSharedPtr<FRepLayout> RepLayout;
+	FRepStateStaticBuffer PropertyData;
 };
 
+typedef TArray<struct FDemoSavedRepObjectState> FDemoSavedPropertyState;
+
 /**
  * Simulated network driver for recording and playing back game sessions.
  */
@@ -276,7 +296,7 @@ class ENGINE_API UDemoNetDriver : public UNetDriver
 	int32		RecordCountSinceFlush;
 
 	/** When we save a checkpoint, we remember all of the actors that need a checkpoint saved out by adding them to this list */
-	TSet< TWeakObjectPtr< AActor > > PendingCheckpointActors;
+	TArray< TWeakObjectPtr< AActor > > PendingCheckpointActors;
 
 	/** Net startup actors that need to be destroyed after checkpoints are loaded */
 	TSet< FString >									DeletedNetStartupActors;
@@ -302,6 +322,9 @@ class ENGINE_API UDemoNetDriver : public UNetDriver
 	void		TickCheckpoint();
 	bool		LoadCheckpoint( FArchive* GotoCheckpointArchive, int64 GotoCheckpointSkipExtraTimeInMS );
 
+	/** Returns true if we're in the process of saving a checkpoint. */
+	bool		IsSavingCheckpoint() const;
+
 	void		SaveExternalData( FArchive& Ar );
 	void		LoadExternalData( FArchive& Ar, const float TimeSeconds );
 
@@ -321,11 +344,17 @@ class ENGINE_API UDemoNetDriver : public UNetDriver
 	/** PlaybackPackets are used to buffer packets up when we read a demo frame, which we can then process when the time is right */
 	TArray< FPlaybackPacket > PlaybackPackets;
 
-	/** All unique streaming levels since recording started */
+	/**
+	 * During recording, all unique streaming levels since recording started.
+	 * During playback, all streaming level instances we've created.
+	 */
 	TSet< TWeakObjectPtr< UObject > >	UniqueStreamingLevels;
 
-	/** Streaming levels waiting to be saved next frame */
-	TArray< UObject* >					NewStreamingLevelsThisFrame;
+	/**
+	 * During recording, streaming levels waiting to be saved next frame.
+	 * During playback, streaming levels that have recently become visible.
+	 */
+	TSet< TWeakObjectPtr< UObject > >	NewStreamingLevelsThisFrame;
 
 	bool bRecordMapChanges;
 
@@ -413,6 +442,7 @@ public:
 	virtual bool ShouldClientDestroyTearOffActors() const override;
 	virtual bool ShouldSkipRepNotifies() const override;
 	virtual bool ShouldQueueBunchesForActorGUID(FNetworkGUID InGUID) const override;
+	virtual bool ShouldIgnoreRPCs() const override;
 	virtual FNetworkGUID GetGUIDForActor(const AActor* InActor) const override;
 	virtual AActor* GetActorForGUID(FNetworkGUID InGUID) const override;
 	virtual bool ShouldReceiveRepNotifiesForObject(UObject* Object) const override;
@@ -450,6 +480,12 @@ public:
 
 	const TArray<FLevelNameAndTime>& GetLevelNameAndTimeList() { return LevelNamesAndTimes; }
 
+	/** Returns the replicated state of every object on a current actor channel. Use the result to compare in DiffReplicatedProperties. */
+	FDemoSavedPropertyState SavePropertyState() const;
+
+	/** Compares the values of replicated properties stored in State with the current values of the object replicators. Logs and returns true if there were any differences. */
+	bool ComparePropertyState(const FDemoSavedPropertyState& State) const;
+
 public:
 
 	UPROPERTY()
@@ -472,9 +508,20 @@ public:
 	bool ConditionallyProcessPlaybackPackets();
 	void ProcessAllPlaybackPackets();
 	bool ReadPacket( FArchive& Archive, uint8* OutReadBuffer, int32& OutBufferSize, const int32 MaxBufferSize );
-	bool ReadDemoFrameIntoPlaybackPackets( FArchive& Ar );
 	bool ConditionallyReadDemoFrameIntoPlaybackPackets( FArchive& Ar );
-	bool ProcessPacket( uint8* Data, int32 Count );
+
+	bool ProcessPacket( const uint8* Data, int32 Count );
+	bool ProcessPacket( const FPlaybackPacket& PlaybackPacket )
+	{
+		bool Result = true;
+		if (!ShouldSkipPlaybackPacket(PlaybackPacket))
+		{
+			ProcessPacket(PlaybackPacket.Data.GetData(), PlaybackPacket.Data.Num());
+			LastProcessedPacketTime = PlaybackPacket.TimeSeconds;
+		}
+
+		return Result;
+	}
 
 	void WriteDemoFrameFromQueuedDemoPackets( FArchive& Ar, TArray<FQueuedDemoPacket>& QueuedPackets );
 	void WritePacket( FArchive& Ar, uint8* Data, int32 Count );
@@ -487,10 +534,13 @@ public:
 	void AddEvent(const FString& Group, const FString& Meta, const TArray<uint8>& Data);
 	void EnumerateEvents(const FString& Group, FEnumerateEventsCompleteDelegate& EnumerationCompleteDelegate);
 	void RequestEventData(const FString& EventID, FOnRequestEventDataComplete& RequestEventDataCompleteDelegate);
-	virtual bool IsFastForwarding() { return bIsFastForwarding; }
+	bool IsFastForwarding() const { return bIsFastForwarding; }
 
 	FReplayExternalDataArray* GetExternalDataArrayForObject( UObject* Object );
 
+	bool ReadDemoFrameIntoPlaybackPackets(FArchive& Ar, TArray<FPlaybackPacket>& Packets, const bool bForLevelFastForward, float* OutTime);
+	bool ReadDemoFrameIntoPlaybackPackets(FArchive& Ar) { return ReadDemoFrameIntoPlaybackPackets(Ar, PlaybackPackets, false, nullptr); }
+
 	/**
 	 * Adds a join-in-progress user to the set of users associated with the currently recording replay (if any)
 	 *
@@ -507,6 +557,7 @@ public:
 	void ClearReplayTasks();
 	bool ProcessReplayTasks();
 	bool IsNamedTaskInQueue( const FName& Name ) const;
+	FName GetNextQueuedTaskName() const;
 
 	/** If a channel is associated with Actor, adds the channel's GUID to the list of GUIDs excluded from queuing bunches during scrubbing. */
 	void AddNonQueuedActorForScrubbing(AActor const* Actor);
@@ -522,6 +573,7 @@ public:
 	void PendingNetGameLoadMapCompleted();
 	
 	virtual void NotifyActorDestroyed( AActor* ThisActor, bool IsSeamlessTravel=false ) override;
+	virtual void NotifyActorLevelUnloaded( AActor* Actor ) override;
 
 	/** Call this function during playback to track net startup actors that need a hard reset when scrubbing, which is done by destroying and then re-spawning */
 	virtual void QueueNetStartupActorForRollbackViaDeletion( AActor* Actor );
@@ -537,6 +589,182 @@ public:
 	/** Returns true if TickFlush can be called in parallel with the Slate tick. */
 	bool ShouldTickFlushAsyncEndOfFrame() const;
 
+	/** Returns whether or not this replay was recorded / is playing with Level Streaming fixes. */
+	bool HasLevelStreamingFixes() const
+	{
+		return bHasLevelStreamingFixes;
+	}
+
+	/**
+	 * Called when a new ActorChannel is opened, before the Actor is notified.
+	 *
+	 * @param Channel	The channel associated with the actor.
+	 * @param Actor		The actor that was recently serialized.
+	 */
+	void PreNotifyActorChannelOpen(UActorChannel* Channel, AActor* Actor);
+
+private:
+
+	void CleanupOutstandingRewindActors();
+
+	// Tracks actors that will need to be rewound during scrubbing.
+	// This list should always be empty outside of scrubbing.
+	TSet<FNetworkGUID> TrackedRewindActorsByGUID;
+
+	/**
+	 * Helps keeps tabs on what levels are Ready, Have Seen data, Level Name, and Index into the main status list.
+	 *
+	 * A Level is not considered ready until the following criteria are met:
+	 *	- UWorld::AddToWorld has been called, signifying the level is both Loaded and Visible (in the streaming sense).
+	 *	- Either:
+	 *		No packets of data have been processed for the level (yet),
+	 *		OR The level has been fully fast-forwarded.
+	 *
+	 * A level is marked as Seen once the replay has seen a packet marked for the level.
+	 */
+	struct FLevelStatus
+	{
+		FLevelStatus(const FString& LevelPackageName) :
+			LevelName(LevelPackageName),
+			LevelIndex(INDEX_NONE),
+			bIsReady(false),
+			bHasBeenSeen(false)
+		{
+		}
+
+		// Level name.
+		FString LevelName;
+
+		// Level index (in AllLevelStatuses).
+		int32 LevelIndex;
+
+		// Whether or not the level is ready to receive streaming data.
+		bool bIsReady;
+
+		// Whether or not we've seen replicated data for the level. Only set during playback.
+		bool bHasBeenSeen;
+	};
+
+	// Tracks all available level statuses.
+	// When Recording, this will be in order of replication, and all statuses will be assumed Seen and Visible (even if unmarked).
+	// During Playback, there's no guaranteed order. Levels will be added either when they are added to the world, or when we handle the first
+	// frame containing replicated data.
+	// Use SeenLevelStatuses and LevelStatusesByName for querying.
+	TArray<FLevelStatus> AllLevelStatuses;
+
+	// Since Arrays are dynamically allocated, we can't just hold onto pointers.
+	// If we tried, the underlying memory could be moved without us knowing.
+	// Therefore, we track the Index into the array which should be independent of allocation.
+
+	// Index of level status (in AllLevelStatuses list).
+	TMap<FString, int32> LevelStatusesByName;
+
+	// List of seen level statuses indices (in ALlLevelStatuses).
+	TArray<int32> SeenLevelStatuses;
+
+	// Time of the last packet we've processed (in seconds).
+	float LastProcessedPacketTime;
+
+	// Time of the last frame we've read (in seconds).
+	float LatestReadFrameTime;
+
+	// Whether or not the Streaming Level Fixes are enabled for capture or playback.
+	bool bHasLevelStreamingFixes;
+
+	// Levels that are currently pending for fast forward.
+	// Using raw pointers, because we manually keep when levels are added and removed.
+	TMap<class ULevel*, TSet<TWeakObjectPtr<class AActor>>> LevelsPendingFastForward;
+
+	// Pairs of Level Indices to the remaining number of actors that need to be processed for a given Demo Frame.
+	// Only used during recording.
+	TArray<TPair<int32, int32>> NumActorsToProcessForLevel;
+
+	// Only used during recording.
+	uint32 NumLevelsAddedThisFrame;
+
+	FLevelStatus& FindOrAddLevelStatus(const FString& LevelPackageName)
+	{
+		if (int32* LevelStatusIndex = LevelStatusesByName.Find(LevelPackageName))
+		{
+			return AllLevelStatuses[*LevelStatusIndex];
+		}
+
+		const int32 Index = AllLevelStatuses.Emplace(LevelPackageName);
+		AllLevelStatuses[Index].LevelIndex = Index;
+
+		LevelStatusesByName.Add(LevelPackageName, Index);
+		NumLevelsAddedThisFrame++;
+
+		return AllLevelStatuses[Index];
+	}
+
+	FLevelStatus& GetLevelStatus(const int32 SeenLevelIndex)
+	{
+		return AllLevelStatuses[SeenLevelStatuses[SeenLevelIndex - 1]];
+	}
+
+	FLevelStatus& GetLevelStatus(const FString& LevelPackageName)
+	{
+		return AllLevelStatuses[LevelStatusesByName[LevelPackageName]];
+	}
+
+	// Determines whether or not a packet should be skipped, based on it's level association.
+	bool ShouldSkipPlaybackPacket(const FPlaybackPacket& Packet);
+
+	void ResetLevelStatuses();
+	void ClearLevelStreamingState()
+	{
+		AllLevelStatuses.Empty();
+		LevelStatusesByName.Empty();
+		SeenLevelStatuses.Empty();
+		LevelsPendingFastForward.Empty();
+		NumActorsToProcessForLevel.Empty();
+		NumLevelsAddedThisFrame = 0;
+	}
+
+	/**
+	 * Replicates the given prioritized actors, so their packets can be captured for recording.
+	 * This should be used for normal frame recording.
+	 * @see ReplicateCheckpointActor for recording during checkpoints.
+	 *
+	 * @param ToReplicate	The actors to replicate.
+	 * @param Params		Implementation specific params necessary to replicate the actor.
+	 *
+	 * @return True if there is time remaining to replicate more actors. False otherwise.
+	 */
+	bool ReplicatePrioritizedActors(const TArray<FActorPriority>& ToReplicate, const class FRepActorsParams& Params);
+	bool ReplicatePrioritizedActors(const TArray<FActorPriority*>& ToReplicate, const class FRepActorsParams& Params);
+	bool ReplicatePrioritizedActor(TArray<AActor*>& ActorsToGoDormant, const FActorPriority& ActorPriority, const class FRepActorsParams& Params);
+
+	/**
+	* Replicates the given prioritized actors, so their packets can be captured for recording.
+	* This should be used for normal frame recording.
+	* @see ReplicateCheckpointActor for recording during checkpoints.
+	*
+	* @param ToReplicate	The actors to replicate.
+	* @param RepStartTime	The start time for replication.
+	*
+	* @return True if there is time remaining to replicate more Actors, False otherwise.
+	*/
+	bool ReplicateCheckpointActor(AActor* ToReplicate, UDemoNetConnection* ClientConnection, class FRepActorsCheckpointParams& Params);
+
+	friend class FPendingTaskHelper;
+
+	// Manages basic setup of newly visible levels, and queuing a FastForward task if necessary.
+	void PrepFastForwardLevels();
+
+	// Performs the logic for actually fast-forwarding a level.
+	bool FastForwardLevels(FArchive* CheckpointArchive, int64 ExtraTime);
+
+	// Hooks used to determine when levels are streamed in, streamed out, or if there's a map change.
+	virtual void OnLevelAddedToWorld(ULevel* Level, UWorld* World) override;
+	virtual void OnLevelRemovedFromWorld(ULevel* Level, UWorld* World) override;
+	void OnPostLoadMapWithWorld(UWorld* World);
+
+	// These should only ever be called when recording.
+	TUniquePtr<class FScopedPacketManager> ConditionallyCreatePacketManager(ULevel& Level);
+	TUniquePtr<class FScopedPacketManager> ConditionallyCreatePacketManager(int32 LevelIndex);
+
 protected:
 	/** allows subclasses to write game specific data to demo header which is then handled by ProcessGameSpecificDemoHeader */
 	virtual void WriteGameSpecificDemoHeader(TArray<FString>& GameSpecificData)
diff --git a/Engine/Source/Runtime/Engine/Classes/Engine/EngineTypes.h b/Engine/Source/Runtime/Engine/Classes/Engine/EngineTypes.h
--- a/Engine/Source/Runtime/Engine/Classes/Engine/EngineTypes.h
+++ b/Engine/Source/Runtime/Engine/Classes/Engine/EngineTypes.h
@@ -3276,6 +3276,8 @@ struct FReplicationFlags
 			uint32 bRepPhysics:1;
 			/** True if this actor is replicating on a replay connection. */
 			uint32 bReplay:1;
+			/** True if this actor's RPCs should be ignored. */
+			uint32 bIgnoreRPCs:1;
 		};
 
 		uint32	Value;
diff --git a/Engine/Source/Runtime/Engine/Classes/Engine/Level.h b/Engine/Source/Runtime/Engine/Classes/Engine/Level.h
--- a/Engine/Source/Runtime/Engine/Classes/Engine/Level.h
+++ b/Engine/Source/Runtime/Engine/Classes/Engine/Level.h
@@ -569,7 +569,7 @@ public:
 	bool HasVisibilityRequestPending() const;
 
 	/** Whether the level is currently pending being made invisible or visible.				*/
-	bool HasVisibilityChangeRequestPending() const;
+	ENGINE_API bool HasVisibilityChangeRequestPending() const;
 
 	// Event on level transform changes
 	DECLARE_MULTICAST_DELEGATE_OneParam(FLevelTransformEvent, const FTransform&);
@@ -631,7 +631,7 @@ private:
 
 public:
 	// Used internally to determine which actors should go on the world's NetworkActor list
-	static bool IsNetActor(const AActor* Actor);
+	ENGINE_API static bool IsNetActor(const AActor* Actor);
 
 	/** Populate an entry for Actor in the DestroyedReplicatedStaticActors list */
 	void CreateReplicatedDestructionInfo(AActor* const Actor);
diff --git a/Engine/Source/Runtime/Engine/Classes/Engine/NetConnection.h b/Engine/Source/Runtime/Engine/Classes/Engine/NetConnection.h
--- a/Engine/Source/Runtime/Engine/Classes/Engine/NetConnection.h
+++ b/Engine/Source/Runtime/Engine/Classes/Engine/NetConnection.h
@@ -18,6 +18,7 @@
 #include "Engine/Channel.h"
 #include "ProfilingDebugging/Histogram.h"
 #include "Containers/ArrayView.h"
+#include "ReplicationDriver.h"
 
 #include "NetConnection.generated.h"
 
@@ -28,6 +29,8 @@ class StatelessConnectHandlerComponent;
 class UActorChannel;
 class UChildConnection;
 
+typedef TMap<TWeakObjectPtr<AActor>, UActorChannel*, FDefaultSetAllocator, TWeakObjectPtrMapKeyFuncs<TWeakObjectPtr<AActor>, UActorChannel*>> FActorChannelMap;
+
 /*-----------------------------------------------------------------------------
 	Types.
 -----------------------------------------------------------------------------*/
@@ -386,21 +389,128 @@ public:
 	int32			LogCallCount;
 	int32			LogSustainedCount;
 
+	// ----------------------------------------------
+	// Actor Channel Accessors
+	// ----------------------------------------------
+
+	void RemoveActorChannel(AActor* Actor)
+	{
+		ActorChannels.Remove(Actor);
+		if (ReplicationConnectionDriver)
+		{
+			ReplicationConnectionDriver->NotifyActorChannelRemoved(Actor);
+		}
+	}
+
+	void AddActorChannel(AActor* Actor, UActorChannel* Channel)
+	{
+		ActorChannels.Add(Actor, Channel);
+		if (ReplicationConnectionDriver)
+		{
+			ReplicationConnectionDriver->NotifyActorChannelAdded(Actor, Channel);
+		}
+	}
+
+	UActorChannel* FindActorChannelRef(const TWeakObjectPtr<AActor>& Actor)
+	{
+		return ActorChannels.FindRef(Actor);
+	}
+
+	UActorChannel** FindActorChannel(const TWeakObjectPtr<AActor>& Actor)
+	{
+		return ActorChannels.Find(Actor);
+	}
+
+	int32 ActorChannelsNum() const
+	{
+		return ActorChannels.Num();
+	}
+
+	FActorChannelMap::TConstIterator ActorChannelConstIterator() const
+	{
+		return ActorChannels.CreateConstIterator();
+	}
+
+	const FActorChannelMap& ActorChannelMap() const
+	{
+		return ActorChannels;
+	}
+
+	UReplicationConnectionDriver* GetReplicationConnectionDriver()
+	{
+		return ReplicationConnectionDriver;
+	}
+
+	void SetReplicationConnectionDriver(UReplicationConnectionDriver* NewReplicationConnectionDriver)
+	{
+		ReplicationConnectionDriver = NewReplicationConnectionDriver;
+	}
+
+private:
 	/** @todo document */
-	TMap<TWeakObjectPtr<AActor>, UActorChannel*, FDefaultSetAllocator, TWeakObjectPtrMapKeyFuncs<TWeakObjectPtr<AActor>, UActorChannel*>> ActorChannels;
+	FActorChannelMap ActorChannels;
 
-	/** This holds a list of actor channels that want to fully shutdown, but need to continue processing bunches before doing so */
-	TMap<FNetworkGUID, TArray<class UActorChannel*>> KeepProcessingActorChannelBunchesMap;
+	UReplicationConnectionDriver* ReplicationConnectionDriver;
 
-	/** A list of replicators that belong to recently dormant actors/objects */
-	TMap< TWeakObjectPtr< UObject >, TSharedRef< FObjectReplicator > > DormantReplicatorMap;
+
+public:
+
+	void AddDestructionInfo(FActorDestructionInfo* DestructionInfo)
+	{
+		if (ReplicationConnectionDriver)
+		{
+			ReplicationConnectionDriver->NotifyAddDestructionInfo(DestructionInfo);
+		}
+		else
+		{
+			DestroyedStartupOrDormantActorGUIDs.Add(DestructionInfo->NetGUID);
+		}
+	}
+
+	void RemoveDestructionInfo(FActorDestructionInfo* DestructionInfo)
+	{
+		if (ReplicationConnectionDriver)
+		{
+			ReplicationConnectionDriver->NotifyRemoveDestructionInfo(DestructionInfo);
+		}
+		else
+		{
+			DestroyedStartupOrDormantActorGUIDs.Remove(DestructionInfo->NetGUID);
+		}
+	}
+	
+	void ResetDestructionInfos()
+	{
+		if (ReplicationConnectionDriver)
+		{
+			ReplicationConnectionDriver->NotifyResetDestructionInfo();
+		}
+		else
+		{
+			DestroyedStartupOrDormantActorGUIDs.Reset(); 
+		}
+	}
+
+	TSet<FNetworkGUID>& GetDestroyedStartupOrDormantActorGUIDs() { return DestroyedStartupOrDormantActorGUIDs; }
+
+private:
 
 	/** The server adds GUIDs to this set for each destroyed actor that does not have a channel
 	 *  but that the client still knows about: startup, dormant, or recently dormant set.
 	 *  This set is also populated from the UNetDriver for clients who join-in-progress, so that they can destroy any
 	 *  startup actors that the server has already destroyed.
 	 */
-	TSet<FNetworkGUID>	DestroyedStartupOrDormantActors;
+	TSet<FNetworkGUID>	DestroyedStartupOrDormantActorGUIDs;
+
+public:
+
+	/** This holds a list of actor channels that want to fully shutdown, but need to continue processing bunches before doing so */
+	TMap<FNetworkGUID, TArray<class UActorChannel*>> KeepProcessingActorChannelBunchesMap;
+
+	/** A list of replicators that belong to recently dormant actors/objects */
+	TMap< TWeakObjectPtr< UObject >, TSharedRef< FObjectReplicator > > DormantReplicatorMap;
+
+	
 
 	ENGINE_API FName GetClientWorldPackageName() const { return ClientWorldPackageName; }
 
@@ -816,6 +926,12 @@ public:
 	/** Returns the online platform name for the player on this connection. Only valid for client connections on servers. */
 	ENGINE_API FName GetPlayerOnlinePlatformName() const { return PlayerOnlinePlatformName; }
 
+	/**
+	 * Sets whether or not we should ignore bunches that would attempt to open channels that are already open.
+	 * Should only be used with InternalAck.
+	 */
+	void SetIgnoreAlreadyOpenedChannels(bool bInIgnoreAlreadyOpenedChannels);
+
 protected:
 
 	void CleanupDormantActorState();
@@ -853,6 +969,10 @@ private:
 
 	/** A map of class names to arrays of time differences between replication of actors of that class for each connection */
 	TMap<FString, TArray<float>> ActorsStarvedByClassTimeMap;
+
+	/** Tracks channels that we should ignore when handling special demo data. */
+	TMap<int32, FNetworkGUID> IgnoringChannels;
+	bool bIgnoreAlreadyOpenedChannels;
 };
 
 
@@ -911,3 +1031,18 @@ struct FScopedNetConnectionSettings
 	FNetConnectionSettings OldSettings;
 	bool ShouldApply;
 };
+
+/** A fake connection that will absorb traffic and auto ack every packet. Useful for testing scaling. Use net.SimulateConnections command to add at runtime. */
+UCLASS(transient, config=Engine)
+class ENGINE_API USimulatedClientNetConnection
+	: public UNetConnection
+{
+	GENERATED_UCLASS_BODY()
+public:
+
+	virtual void LowLevelSend(void* Data, int32 CountBytes, int32 CountBits) override { }
+	void HandleClientPlayer( APlayerController* PC, UNetConnection* NetConnection ) override;
+	virtual FString LowLevelGetRemoteAddress(bool bAppendPort=false) override { return FString(); }
+	virtual bool ClientHasInitializedLevelFor(const AActor* TestActor) const { return true; }
+};
+
diff --git a/Engine/Source/Runtime/Engine/Classes/Engine/NetDriver.h b/Engine/Source/Runtime/Engine/Classes/Engine/NetDriver.h
--- a/Engine/Source/Runtime/Engine/Classes/Engine/NetDriver.h
+++ b/Engine/Source/Runtime/Engine/Classes/Engine/NetDriver.h
@@ -27,9 +27,13 @@ class FReplicationChangelistMgr;
 class FVoicePacket;
 class StatelessConnectHandlerComponent;
 class UNetConnection;
+class UReplicationDriver;
 struct FNetworkObjectInfo;
 
 extern ENGINE_API TAutoConsoleVariable<int32> CVarNetAllowEncryption;
+extern ENGINE_API int32 GNumSaturatedConnections;
+extern ENGINE_API int32 GNumSharedSerializationHit;
+extern ENGINE_API int32 GNumSharedSerializationMiss;
 
 // Delegates
 
@@ -206,6 +210,9 @@ public:
 	UPROPERTY(Config)
 	FString NetConnectionClassName;
 
+	UPROPERTY(Config)
+	FString ReplicationDriverClassName;
+
 	/** @todo document */
 	UPROPERTY(Config)
 	int32 MaxDownloadSize;
@@ -302,6 +309,9 @@ public:
 	UPROPERTY()
 	UClass* NetConnectionClass;
 
+	UPROPERTY()
+	UClass* ReplicationDriverClass;
+
 	/** @todo document */
 	UPROPERTY()
 	UProperty* RoleProperty;
@@ -430,7 +440,7 @@ public:
 	 *  to startup actors (because they won't have an associated channel), and this map stores those
 	 *  FActorDestructionInfos also.
 	 */
-	TMap<FNetworkGUID, FActorDestructionInfo>	DestroyedStartupOrDormantActors;
+	TMap<FNetworkGUID, TUniquePtr<FActorDestructionInfo>>	DestroyedStartupOrDormantActors;
 
 	/** The server adds an entry into this map for every startup actor that has been renamed, and will
 	 *  always map from current name to original name
@@ -452,7 +462,7 @@ public:
 	TSharedPtr< FRepLayout >	GetObjectClassRepLayout( UClass * InClass );
 
 	/** Creates if necessary, and returns a FRepLayout that maps to the passed in UFunction */
-	TSharedPtr<FRepLayout>		GetFunctionRepLayout( UFunction * Function );
+	ENGINE_API TSharedPtr<FRepLayout> GetFunctionRepLayout( UFunction * Function );
 
 	/** Creates if necessary, and returns a FRepLayout that maps to the passed in UStruct */
 	TSharedPtr<FRepLayout>		GetStructRepLayout( UStruct * Struct );
@@ -584,6 +594,9 @@ public:
 	/** Initializes the net connection class to use for new connections */
 	ENGINE_API virtual bool InitConnectionClass(void);
 
+	/** Initialized the replication driver class to use for this driver */
+	ENGINE_API virtual bool InitReplicationDriverClass();
+
 	/** Shutdown all connections managed by this net driver */
 	ENGINE_API virtual void Shutdown();
 
@@ -624,6 +637,22 @@ public:
 	 */
 	ENGINE_API virtual void ProcessRemoteFunction(class AActor* Actor, class UFunction* Function, void* Parameters, struct FOutParmRec* OutParms, struct FFrame* Stack, class UObject* SubObject = nullptr );
 
+
+	enum ENGINE_API ERemoteFunctionSendPolicy
+	{		
+		/** Unreliable multicast are queued. Everything else is send immediately */
+		Default, 
+
+		/** Bunch is send immediately no matter what */
+		ForceSend,
+
+		/** Bunch is queued until next actor replication, no matter what */
+		ForceQueue,
+	};	
+
+	/** Process a remote function on given actor channel. This is called by ::ProcessRemoteFunction.*/
+	ENGINE_API void ProcessRemoteFunctionForChannel(UActorChannel* Ch, const class FClassNetCache* ClassCache, const FFieldNetCache* FieldCache, UObject* TargetObj, UNetConnection* Connection,  UFunction* Function, void* Parms, FOutParmRec* OutParms, FFrame* Stack, const bool IsServer, const ERemoteFunctionSendPolicy SendPolicy = ERemoteFunctionSendPolicy::Default);
+
 	/** handle time update */
 	ENGINE_API virtual void TickDispatch( float DeltaTime );
 
@@ -682,8 +711,16 @@ public:
 	bool HandleNetDumpDormancy( const TCHAR* Cmd, FOutputDevice& Ar );
 #endif
 
+	// ---------------------------------------------------------------
+	//	Game code API for updating server Actor Replication State
+	// ---------------------------------------------------------------
+
+	ENGINE_API void ForceNetUpdate(AActor* Actor);
+
 	/** Flushes actor from NetDriver's dormancy list, but does not change any state on the Actor itself */
-	ENGINE_API void FlushActorDormancy(class AActor *Actor);
+	ENGINE_API void FlushActorDormancy(AActor *Actor, bool bWasDormInitial=false);
+
+	ENGINE_API void NotifyActorDormancyChange(AActor* Actor, ENetDormancy OldDormancyState);
 
 	/** Forces properties on this actor to do a compare for one frame (rather than share shadow state) */
 	ENGINE_API void ForcePropertyCompare( AActor* Actor );
@@ -691,16 +728,27 @@ public:
 	/** Force this actor to be relevant for at least one update */
 	ENGINE_API void ForceActorRelevantNextUpdate(AActor* Actor);
 
+	/** Tells the net driver about a networked actor that was spawned */
+	ENGINE_API void AddNetworkActor(AActor* Actor);
+
 	/** Called when a spawned actor is destroyed. */
 	ENGINE_API virtual void NotifyActorDestroyed( AActor* Actor, bool IsSeamlessTravel=false );
 
 	/** Called when an actor is renamed. */
 	ENGINE_API virtual void NotifyActorRenamed(AActor* Actor, FName PreviousName);
 
-	ENGINE_API virtual void NotifyStreamingLevelUnload( ULevel* );
+	ENGINE_API void RemoveNetworkActor(AActor* Actor);
 
 	ENGINE_API virtual void NotifyActorLevelUnloaded( AActor* Actor );
-	
+
+	ENGINE_API virtual void NotifyActorTearOff(AActor* Actor);
+
+	// ---------------------------------------------------------------
+	//
+	// ---------------------------------------------------------------	
+
+	ENGINE_API virtual void NotifyStreamingLevelUnload( ULevel* );
+
 	/** creates a child connection and adds it to the given parent connection */
 	ENGINE_API virtual class UChildConnection* CreateChild(UNetConnection* Parent);
 
@@ -765,6 +813,9 @@ public:
 	/** Returns true if actor channels with InGUID should queue up bunches, even if they wouldn't otherwise be queued. */
 	virtual bool ShouldQueueBunchesForActorGUID(FNetworkGUID InGUID) const { return false; }
 
+	/** Returns whether or not RPCs processed by this driver should be ignored. */
+	virtual bool ShouldIgnoreRPCs() const { return false; }
+
 	/** Returns the existing FNetworkGUID of InActor, if it has one. */
 	virtual FNetworkGUID GetGUIDForActor(const AActor* InActor) const { return FNetworkGUID(); }
 
@@ -828,17 +879,30 @@ public:
 	/** Sets the level ID/PIE instance ID for this netdriver to use. */
 	ENGINE_API void SetDuplicateLevelID(const int32 InDuplicateLevelID) { DuplicateLevelID = InDuplicateLevelID; }
 
-protected:
+	/** Explicitly sets the ReplicationDriver instance (you instantiate it and initialize it). Shouldn't be done during gameplay: ok to do in GameMode startup or via console commands for testing.  */
+	ENGINE_API void SetReplicationDriver(UReplicationDriver* NewReplicationManager);
+
+	ENGINE_API UReplicationDriver* GetReplicationDriver() { return ReplicationDriver; }
+
+	template<class T>
+	T* GetReplicationDriver() { return Cast<T>(ReplicationDriver); }
+
+	void RemoveClientConnection(UNetConnection* ClientConnectionToRemove);
 
 	/** Adds (fully initialized, ready to go) client connection to the ClientConnections list + any other game related setup */
 	ENGINE_API void	AddClientConnection(UNetConnection * NewConnection);
 
+	ENGINE_API void NotifyActorFullyDormantForConnection(AActor* Actor, UNetConnection* Connection);
+
+	/** Returns true if this actor is considered to be in a loaded level */
+	ENGINE_API virtual bool IsLevelInitializedForActor( const AActor* InActor, const UNetConnection* InConnection ) const;
+
+protected:
+
 	/** Register all TickDispatch, TickFlush, PostTickFlush to tick in World */
 	ENGINE_API void RegisterTickEvents(class UWorld* InWorld);
 	/** Unregister all TickDispatch, TickFlush, PostTickFlush to tick in World */
 	ENGINE_API void UnregisterTickEvents(class UWorld* InWorld);
-	/** Returns true if this actor is considered to be in a loaded level */
-	ENGINE_API virtual bool IsLevelInitializedForActor( const AActor* InActor, const UNetConnection* InConnection ) const;
 
 #if WITH_SERVER_CODE
 	/**
@@ -853,8 +917,12 @@ protected:
 	/** Used to handle any NetDriver specific cleanup once a level has been removed from the world. */
 	ENGINE_API virtual void OnLevelRemovedFromWorld(class ULevel* Level, class UWorld* World);
 
-	/** Handle that tracks OnLevelRemovedFromWorld. */
+	/** Used to handle any NetDriver specific setup when a level has been added to the world. */
+	ENGINE_API virtual void OnLevelAddedToWorld(class ULevel* Level, class UWorld* World);
+
+	/** Handles that track our LevelAdded / Removed delegates. */
 	FDelegateHandle OnLevelRemovedFromWorldHandle;
+	FDelegateHandle OnLevelAddedToWorldHandle;
 
 private:
 	FActorDestructionInfo* CreateDestructionInfo(UNetDriver* NetDriver, AActor* ThisActor, FActorDestructionInfo *DestructionInfo);
@@ -862,6 +930,11 @@ private:
 	void CreateReplicatedStaticActorDestructionInfo(UNetDriver* NetDriver, ULevel* Level, const FReplicatedStaticActorDestructionInfo& Info);
 
 
+	void FlushActorDormancyInternal(AActor *Actor);
+
+	UPROPERTY(transient)
+	UReplicationDriver* ReplicationDriver;
+
 	/** Stores the list of objects to replicate into the replay stream. This should be a TUniquePtr, but it appears the generated.cpp file needs the full definition of the pointed-to type. */
 	TSharedPtr<FNetworkObjectList> NetworkObjects;
 
diff --git a/Engine/Source/Runtime/Engine/Classes/Engine/NetworkObjectList.h b/Engine/Source/Runtime/Engine/Classes/Engine/NetworkObjectList.h
--- a/Engine/Source/Runtime/Engine/Classes/Engine/NetworkObjectList.h
+++ b/Engine/Source/Runtime/Engine/Classes/Engine/NetworkObjectList.h
@@ -122,7 +122,7 @@ public:
 	 *
 	 * If info cannot be found or created, nullptr will be returned.
 	 */
-	TSharedPtr<FNetworkObjectInfo>* FindOrAdd(AActor* const Actor, const FName NetDriverName);
+	TSharedPtr<FNetworkObjectInfo>* FindOrAdd(AActor* const Actor, const FName NetDriverName, bool* OutWasAdded=nullptr);
 
 	/**
 	 * Attempts to find the Actor's FNetworkObjectInfo.
diff --git a/Engine/Source/Runtime/Engine/Classes/Engine/PackageMapClient.h b/Engine/Source/Runtime/Engine/Classes/Engine/PackageMapClient.h
--- a/Engine/Source/Runtime/Engine/Classes/Engine/PackageMapClient.h
+++ b/Engine/Source/Runtime/Engine/Classes/Engine/PackageMapClient.h
@@ -291,6 +291,9 @@ public:
 	void ReceiveNetGUIDBunch( FInBunch &InBunch );
 	void AppendExportBunches(TArray<FOutBunch *>& OutgoingBunches);
 
+	void AppendExportData(FArchive& Archive);
+	void ReceiveExportData(FArchive& Archive);
+
 	TMap<FNetworkGUID, int32>	NetGUIDExportCountMap;	// How many times we've exported each NetGUID on this connection. Public for ListNetGUIDExports 
 
 	void HandleUnAssignedObject( const UObject* Obj );
@@ -327,14 +330,23 @@ public:
 	void								AddNetFieldExportGroup( const FString& PathName, TSharedPtr< FNetFieldExportGroup > NewNetFieldExportGroup );
 	void								TrackNetFieldExport( FNetFieldExportGroup* NetFieldExportGroup, const int32 NetFieldExportHandle );
 	TSharedPtr< FNetFieldExportGroup >	GetNetFieldExportGroupChecked( const FString& PathName ) const;
-	void								SerializeNetFieldExportGroupMap( FArchive& Ar );
+	void								SerializeNetFieldExportGroupMap( FArchive& Ar, bool bClearPendingExports=true );
 
 protected:
 
 	/** Functions to help with exporting/importing net field export info */
+	DEPRECATED(4.20, "This method is deprecated. Please use AppendNetFieldExports(FArchive*) instead.")
 	void								AppendNetFieldExports( TArray<FOutBunch *>& OutgoingBunches );
+	void								AppendNetFieldExports( FArchive& Archive );
+
+	DEPRECATED(4.20, "This method is deprecated. Please use ReceiveNetFieldExports(FArchive*) instead.")
 	void								ReceiveNetFieldExports( FInBunch &InBunch );
+	void								ReceiveNetFieldExports( FArchive& Archive );
+
+	void AppendNetExportGUIDs(FArchive& Archive);
+	void ReceiveNetExportGUIDs(FArchive& Archive);
 
+	bool	ExportNetGUIDForReplay( FNetworkGUID&, const UObject* Object, FString& PathName, UObject* ObjOuter );
 	bool	ExportNetGUID( FNetworkGUID NetGUID, const UObject* Object, FString PathName, UObject* ObjOuter );
 	void	ExportNetGUIDHeader();
 
@@ -351,6 +363,7 @@ protected:
 
 	bool ObjectLevelHasFinishedLoading(UObject* Obj);
 
+	TArray<TArray<uint8>>				ExportGUIDArchives;
 	TSet< FNetworkGUID >				CurrentExportNetGUIDs;				// Current list of NetGUIDs being written to the Export Bunch.
 	TSet< FNetworkGUID >				CurrentQueuedBunchNetGUIDs;			// List of NetGuids with currently queued bunches
 
diff --git a/Engine/Source/Runtime/Engine/Classes/Engine/World.h b/Engine/Source/Runtime/Engine/Classes/Engine/World.h
--- a/Engine/Source/Runtime/Engine/Classes/Engine/World.h
+++ b/Engine/Source/Runtime/Engine/Classes/Engine/World.h
@@ -2001,6 +2001,16 @@ public:
 	/** Returns true if the actors have been initialized and are ready to start play */
 	bool AreActorsInitialized() const;
 
+	struct FActorsInitializedParams
+	{
+		FActorsInitializedParams(UWorld* InWorld, bool InResetTime) : World(InWorld), ResetTime(InResetTime) {}
+		UWorld* World;
+		bool ResetTime;
+	};
+
+	DECLARE_MULTICAST_DELEGATE_OneParam(FOnWorldInitializedActors, const FActorsInitializedParams&);
+	FOnWorldInitializedActors OnActorsInitialized;
+
 	/** Returns true if gameplay has already started, false otherwise. */
 	bool HasBegunPlay() const;
 
@@ -3372,6 +3382,9 @@ public:
 	// Called when changes in the levels require blueprint actions to be refreshed.
 	static FRefreshLevelScriptActionsEvent RefreshLevelScriptActions;
 #endif
+	
+	// Global Callback after actors have been initialized (on any world)
+	static UWorld::FOnWorldInitializedActors OnWorldInitializedActors;
 
 private:
 	FWorldDelegates() {}
diff --git a/Engine/Source/Runtime/Engine/Classes/GameFramework/Actor.h b/Engine/Source/Runtime/Engine/Classes/GameFramework/Actor.h
--- a/Engine/Source/Runtime/Engine/Classes/GameFramework/Actor.h
+++ b/Engine/Source/Runtime/Engine/Classes/GameFramework/Actor.h
@@ -155,8 +155,14 @@ public:
 	 * If true, this actor is no longer replicated to new clients, and is "torn off" (becomes a ROLE_Authority) on clients to which it was being replicated.
 	 * @see TornOff()
 	 */
+	DEPRECATED(4.20, "Use GetTearOff() or TearOff() functions. This property will become private.")
 	UPROPERTY(Replicated)
-	uint8 bTearOff:1;    
+	uint8 bTearOff:1; 
+
+	/** returns TearOff status */
+	PRAGMA_DISABLE_DEPRECATION_WARNINGS
+	bool GetTearOff() const { return bTearOff; }
+	PRAGMA_ENABLE_DEPRECATION_WARNINGS
 
 	/** Networking - Server - TearOff this actor to stop replication to clients. Will set bTearOff to true. */
 	UFUNCTION(BlueprintCallable, Category=Replication)
@@ -181,6 +187,14 @@ public:
 	UPROPERTY()
 	uint8 bRelevantForNetworkReplays:1;
 
+	/**
+	 * If true, this actor will only be destroyed during scrubbing if the replay is set to a time before the actor existed.
+	 * Otherwise, RewindForReplay will be called if we detect the actor needs to be reset.
+	 * Note, this Actor must not be destroyed by gamecode, and RollbackViaDeletion may not be used. 
+	 */
+	UPROPERTY()
+	uint8 bReplayRewindable:1;
+
 	/**
 	 * Whether we allow this Actor to tick before it receives the BeginPlay event.
 	 * Normally we don't tick actors until after BeginPlay; this setting allows this behavior to be overridden.
@@ -483,6 +497,12 @@ public:
 	 */
 	virtual void PreReplicationForReplay(IRepChangedPropertyTracker & ChangedPropertyTracker);
 
+	/**
+	 * Called on the actor before checkpoint data is applied during a replay.
+	 * Only called if bReplayRewindable is set.
+	 */
+	virtual void RewindForReplay();
+
 	/** Called by the networking system to call PreReplication on this actor and its components using the given NetDriver to find or create RepChangedPropertyTrackers. */
 	void CallPreReplication(UNetDriver* NetDriver);	
 	
diff --git a/Engine/Source/Runtime/Engine/Classes/GameFramework/HUD.h b/Engine/Source/Runtime/Engine/Classes/GameFramework/HUD.h
--- a/Engine/Source/Runtime/Engine/Classes/GameFramework/HUD.h
+++ b/Engine/Source/Runtime/Engine/Classes/GameFramework/HUD.h
@@ -21,6 +21,7 @@ class UMaterialInterface;
 class UTexture;
 
 DECLARE_MULTICAST_DELEGATE_FiveParams(FOnShowDebugInfo, AHUD* /* HUD */, UCanvas* /* Canvas */, const FDebugDisplayInfo& /* DisplayInfo */, float& /* YL */, float& /* YPos */);
+DECLARE_MULTICAST_DELEGATE_TwoParams(FOnHUDPostRender, AHUD* /* HUD */, UCanvas* /* Canvas */);
 
 /** 
  * Base class of the heads-up display. This has a canvas and a debug canvas on which primitives can be drawn.
@@ -497,6 +498,9 @@ public:
 	// Callback allowing external systems to register to show debug info
 	static FOnShowDebugInfo OnShowDebugInfo;
 
+	// Called from ::PostRender. For less player/actor centered debugging
+	static FOnHUDPostRender OnHUDPostRender;
+
 	/** PostRender is the main draw loop. */
 	virtual void PostRender();
 
