Fix a split screen RPC crashin rep graph

diff --git a/Engine/Plugins/Runtime/ReplicationGraph/Source/Private/ReplicationGraph.cpp b/Engine/Plugins/Runtime/ReplicationGraph/Source/Private/ReplicationGraph.cpp
--- a/Engine/Plugins/Runtime/ReplicationGraph/Source/Private/ReplicationGraph.cpp
+++ b/Engine/Plugins/Runtime/ReplicationGraph/Source/Private/ReplicationGraph.cpp
@@ -1169,6 +1169,12 @@ bool UReplicationGraph::ProcessRemoteFunction(class AActor* Actor, UFunction* Fu
 		{
 			return true;
 		}
+
+		// Route RPC calls to actual connection
+		if (Connection->GetUChildConnection())
+		{
+			Connection = ((UChildConnection*)Connection)->Parent;
+		}
 	
 		if (Connection->State == USOCK_Closed)
 		{
