Include all traffic 'til the ActorSpawn was ackknowledged in the RepGraph's separate actor discovery budget

diff --git a/Engine/Plugins/Runtime/ReplicationGraph/Source/Private/ReplicationGraph.cpp b/Engine/Plugins/Runtime/ReplicationGraph/Source/Private/ReplicationGraph.cpp
--- a/Engine/Plugins/Runtime/ReplicationGraph/Source/Private/ReplicationGraph.cpp
+++ b/Engine/Plugins/Runtime/ReplicationGraph/Source/Private/ReplicationGraph.cpp
@@ -1578,7 +1578,7 @@ int64 UReplicationGraph::ReplicateSingleActor(AActor* Actor, FConnectionReplicat
 	}
 
 	// Optional budget for actor discovery traffic
-	if (ActorDiscoveryMaxBitsPerFrame > 0 && bOpenActorChannel )
+	if (ActorDiscoveryMaxBitsPerFrame > 0 && (ActorInfo.Channel && ActorInfo.Channel->SpawnAcked == false) )
 	{
 		if (ConnectionManager.QueuedBitsForActorDiscovery < ActorDiscoveryMaxBitsPerFrame)
 		{
