Try to catch edge cases for Adding and Removing channels in RepGraph when we may have the updated channels for the same actor.

diff --git a/Engine/Plugins/Runtime/ReplicationGraph/Source/Private/ReplicationGraph.cpp b/Engine/Plugins/Runtime/ReplicationGraph/Source/Private/ReplicationGraph.cpp
--- a/Engine/Plugins/Runtime/ReplicationGraph/Source/Private/ReplicationGraph.cpp
+++ b/Engine/Plugins/Runtime/ReplicationGraph/Source/Private/ReplicationGraph.cpp
@@ -120,6 +120,13 @@ static TAutoConsoleVariable<FString> CVarRepGraphConditionalBreakpointActorName(
 // Variable that can be programatically set to a specific actor/connection 
 FActorConnectionPair DebugActorConnectionPair;
 
+/** Used to call Describe on a Connection or Channel, handling the null case. */
+template<typename T>
+static FORCEINLINE FString DescribeSafe(T* Describable)
+{
+	return Describable ? Describable->Describe() : FString(TEXT("None"));
+}
+
 FORCEINLINE bool RepGraphConditionalActorBreakpoint(AActor* Actor, UNetConnection* NetConnection)
 {
 #if !(UE_BUILD_SHIPPING || UE_BUILD_TEST)
@@ -844,16 +851,25 @@ int32 UReplicationGraph::ServerReplicateActors(float DeltaSeconds)
 					UActorChannel* Channel = MapIt.Key();
 					checkSlow(Channel != nullptr);
 					checkSlow(ConnectionActorInfo.Channel != nullptr);
-					ensureMsgf(Channel == ConnectionActorInfo.Channel, TEXT("Channel: %s ConnectionActorInfo.Channel: %s."), *(Channel ? Channel->Describe() : FString(TEXT("None"))), *(ConnectionActorInfo.Channel ? ConnectionActorInfo.Channel->Describe() : FString(TEXT("None"))));
 
-					if (ConnectionActorInfo.ActorChannelCloseFrameNum > 0 && ConnectionActorInfo.ActorChannelCloseFrameNum <= FrameNum)
+					// We check for Channel closing early and bail.
+					// It may be possible when using Dormancy that an Actor's Channel was closed, but a new channel was created
+					// before the original Cleaned Up.
+					if (Channel->Closing)
 					{
-						if (ConnectionActorInfo.Channel->Closing)
-						{
-							UE_LOG(LogReplicationGraph, Verbose, TEXT("NET_ReplicateActors_LookForNonRelevantChannels Channel %s is closing. Skipping."), *GetNameSafe(ConnectionActorInfo.Channel));
-							continue;
-						}
+						UE_LOG(LogReplicationGraph, Verbose, TEXT("NET_ReplicateActors_LookForNonRelevantChannels (key) Channel %s is closing. Skipping."), *Channel->Describe());
+						continue;
+					}
+					else if (ConnectionActorInfo.Channel->Closing)
+					{
+						UE_LOG(LogReplicationGraph, Verbose, TEXT("NET_ReplicateActors_LookForNonRelevantChannels (value) Channel %s is closing. Skipping."), *ConnectionActorInfo.Channel->Describe());
+						continue;
+					}
+
+					ensureMsgf(Channel == ConnectionActorInfo.Channel, TEXT("Channel: %s ConnectionActorInfo.Channel: %s."), *Channel->Describe(), *ConnectionActorInfo.Channel->Describe());
 
+					if (ConnectionActorInfo.ActorChannelCloseFrameNum > 0 && ConnectionActorInfo.ActorChannelCloseFrameNum <= FrameNum)
+					{
 						AActor* Actor = Channel->Actor;
 
 						if (ensureMsgf(Actor,
@@ -1822,7 +1838,10 @@ bool UReplicationGraph::ProcessRemoteFunction(class AActor* Actor, UFunction* Fu
 	UNetConnection* Connection = Actor->GetNetConnection();
 	if (Connection)
 	{
-		if (((Function->FunctionFlags & FUNC_NetReliable) == 0) && !IsConnectionReady(Connection))
+		const bool bIsReliable = EnumHasAnyFlags(Function->FunctionFlags, FUNC_NetReliable);
+
+		// If we're saturated and it's not a reliable multicast, drop it.
+		if (!(bIsReliable || IsConnectionReady(Connection)))
 		{
 			return true;
 		}
@@ -1847,6 +1866,20 @@ bool UReplicationGraph::ProcessRemoteFunction(class AActor* Actor, UFunction* Fu
 				return true;
 			}
 
+			if (UNetReplicationGraphConnection* ConnectionManager = Cast<UNetReplicationGraphConnection>(Connection->GetReplicationConnectionDriver()))
+			{
+				if (FConnectionReplicationActorInfo const * const ConnectionActorInfo = ConnectionManager->ActorInfoMap.Find(Actor))
+				{
+					const bool bIsValid = ensureMsgf(ConnectionActorInfo->Channel == nullptr, TEXT("UReplicationGraph::ProcessRemoteFunction: Trying to send RPC for Actor whose channel is in invalid state %s"), *ConnectionActorInfo->Channel->Describe());
+
+					// Don't try to open a channel for an unreliable RPC whose channel is closing.
+					if (!(bIsReliable || bIsValid))
+					{
+						return true;
+					}
+				}
+			}
+
 			Ch = (UActorChannel *)Connection->CreateChannelByName( NAME_Actor, EChannelCreateFlags::OpenedLocally );
 			Ch->SetChannelActor(Actor);
 		}
@@ -1963,8 +1996,17 @@ void UNetReplicationGraphConnection::NotifyActorChannelAdded(AActor* Actor, clas
 	}
 
 	FConnectionReplicationActorInfo& ActorInfo = ActorInfoMap.FindOrAdd(Actor);
-	ActorInfo.Channel = Channel;
 
+	// The ActorInfoMap may have a channel already.
+	// This may happen in cases like dormancy where new Actor Channels can be created and then
+	// closed multiple times for the same Actor, potentially before receiving CleanUp calls.
+	if (ActorInfo.Channel && Channel != ActorInfo.Channel)
+	{
+		UE_LOG(LogReplicationGraph, Log, TEXT("::NotifyActorChannelAdded. Fixing up stale channel reference Old: %s New: %s"), *ActorInfo.Channel->Describe(), *Channel->Describe());
+		ActorInfoMap.RemoveChannel(ActorInfo.Channel);
+	}
+
+	ActorInfo.Channel = Channel;
 	ActorInfoMap.AddChannel(Actor, Channel);
 }
 
@@ -1989,14 +2031,27 @@ void UNetReplicationGraphConnection::NotifyActorChannelCleanedUp(UActorChannel*
 		FConnectionReplicationActorInfo* ActorInfo = ActorInfoMap.FindByChannel(Channel);
 		if (ActorInfo)
 		{
-			// Remove reference from channel map
-			ActorInfoMap.RemoveChannel(Channel);
-
 			// Note we can't directly remove the entry from ActorInfoMap.ActorMap since we don't have the AActor* to key into that map
 			// But we don't actually have to remove the entry since we no longer iterate through ActorInfoMap.ActorMap in non debug functions.
 			// So all we need to do is clear the runtime/transient data for this actorinfo map. (We want to preserve the dormancy flag and the 
 			// settings we pulled from the FGlobalActorReplicationInfo, but clear the frame counters, etc).
-			ActorInfo->ResetFrameCounters();
+
+			if (Channel == ActorInfo->Channel)
+			{
+				// Only reset our state if we're still the associated channel.
+				ActorInfo->ResetFrameCounters();
+			}
+			else
+			{
+				UE_LOG(LogReplicationGraph, Log, TEXT("::NotifyActorChannelCleanedUp. CleanUp for stale channel reference Old: %s New: %s"), *Channel->Describe(), *DescribeSafe(ActorInfo->Channel));
+			}
+
+			// Remove reference from channel map
+			// We call this last, as it could be the last thing holding onto the underlying
+			// shared pointer and we don't want to try and access potentially garbage memory.
+			// This isn't a big deal for now since FConnectionReplicationActorInfo is just a POD
+			// type, but if that changes it could be a problem.
+			ActorInfoMap.RemoveChannel(Channel);
 		}
 	}
 }
