Additional tracking for dependent actors so they can always be cleaned up when removed from the rep graph

diff --git a/Engine/Plugins/Runtime/ReplicationGraph/Source/Private/ReplicationGraph.cpp b/Engine/Plugins/Runtime/ReplicationGraph/Source/Private/ReplicationGraph.cpp
--- a/Engine/Plugins/Runtime/ReplicationGraph/Source/Private/ReplicationGraph.cpp
+++ b/Engine/Plugins/Runtime/ReplicationGraph/Source/Private/ReplicationGraph.cpp
@@ -1503,7 +1503,8 @@ int64 UReplicationGraph::ReplicateSingleActor(AActor* Actor, FConnectionReplicat
 	}
 
 	const bool bWantsToGoDormant = GlobalActorInfo.bWantsToBeDormant;
-	TActorRepListViewBase<FActorRepList*> DependentActorList(GlobalActorInfo.DependentActorList.RepList.GetReference());
+
+	const FActorRepListRefView& DependentActorList = GlobalActorInfo.GetDependentActorList();
 
 	bool bOpenActorChannel = (ActorInfo.Channel == nullptr);
 
@@ -1618,8 +1619,9 @@ void UReplicationGraph::HandleStarvedActorList(const FPrioritizedRepList& List,
 
 		// Update dependent actor's timeout frame
 		FGlobalActorReplicationInfo& GlobalActorInfo = GlobalActorReplicationInfoMap.Get(RepItem.Actor);
-		TActorRepListViewBase<FActorRepList*> DependentActorList(GlobalActorInfo.DependentActorList.RepList.GetReference());
-		
+
+		const FActorRepListRefView& DependentActorList = GlobalActorInfo.GetDependentActorList();
+
 		if (DependentActorList.IsValid())
 		{
 			const uint32 CloseFrameNum = ActorInfo.ActorChannelCloseFrameNum;
diff --git a/Engine/Plugins/Runtime/ReplicationGraph/Source/Public/ReplicationGraphTypes.h b/Engine/Plugins/Runtime/ReplicationGraph/Source/Public/ReplicationGraphTypes.h
--- a/Engine/Plugins/Runtime/ReplicationGraph/Source/Public/ReplicationGraphTypes.h
+++ b/Engine/Plugins/Runtime/ReplicationGraph/Source/Public/ReplicationGraphTypes.h
@@ -452,9 +452,6 @@ struct FGlobalActorReplicationInfo
 	/** Mirrors AActor::NetDormany > DORM_Awake */
 	bool bWantsToBeDormant = false;
 
-	/** When this actor replicates, we replicate these actors immediately afterwards (they are not gathered/prioritized/etc) */
-	FActorRepListRefView DependentActorList;
-	
 	/** Class default mirrors: state that is initialized directly from class defaults (and can be later changed on a per-actor basis) */
 	FClassReplicationInfo Settings;
 	
@@ -480,6 +477,17 @@ struct FGlobalActorReplicationInfo
 			FastSharedReplicationInfo->CountBytes(Ar);
 		}
 	}
+
+	const FActorRepListRefView& GetDependentActorList() { return DependentActorList; }
+
+	friend struct FGlobalActorReplicationInfoMap;
+
+private:
+	/** When this actor replicates, we replicate these actors immediately afterwards (they are not gathered/prioritized/etc) */
+	FActorRepListRefView DependentActorList;
+
+	/** When this actor is added to the dependent list of a parent, track the parent here */
+	FActorRepListRefView ParentActorList;
 };
 
 /** Templatd struct for mapping UClasses to some data type. The main things this provides is that if a UClass* was not explicitly added, it will climb the class heirachy and find the best match (and then store this for faster lookup next time) */
@@ -662,7 +670,29 @@ struct FGlobalActorReplicationInfoMap
 	}
 
 	/** Removes actor data from map */
-	FORCEINLINE int32 Remove(const FActorRepListType& Actor) { return ActorMap.Remove(Actor); }
+	int32 Remove(const FActorRepListType& Actor)
+	{
+		if (FGlobalActorReplicationInfo* ActorInfo = Find(Actor))
+		{
+			if (ActorInfo->DependentActorList.IsValid())
+			{
+				for (AActor* ChildActor : ActorInfo->DependentActorList)
+				{
+					RemoveDependentActor(Actor, ChildActor);
+				}
+			}
+
+			if (ActorInfo->ParentActorList.IsValid())
+			{
+				for (AActor* ParentActor : ActorInfo->ParentActorList)
+				{
+					RemoveDependentActor(ParentActor, Actor);
+				}
+			}
+		}
+
+		return ActorMap.Remove(Actor);
+	}
 
 	/** Returns ClassInfo for a given class. */
 	FORCEINLINE FClassReplicationInfo& GetClassInfo(UClass* Class) { return ClassMap.GetChecked(Class); }
@@ -692,6 +722,42 @@ struct FGlobalActorReplicationInfoMap
 		ClassMap.CountBytes(Ar);
 	}
 
+	void AddDependentActor(AActor* Parent, AActor* Child)
+	{
+		if (Parent && Child)
+		{
+			if (FGlobalActorReplicationInfo* ParentInfo = Find(Parent))
+			{
+				ParentInfo->DependentActorList.PrepareForWrite();
+				ParentInfo->DependentActorList.ConditionalAdd(Child);
+			}
+
+			if (FGlobalActorReplicationInfo* ChildInfo = Find(Child))
+			{
+				ChildInfo->ParentActorList.PrepareForWrite();
+				ChildInfo->ParentActorList.ConditionalAdd(Parent);
+			}
+		}
+	}
+
+	void RemoveDependentActor(AActor* Parent, AActor* Child)
+	{
+		if (Parent && Child)
+		{
+			if (FGlobalActorReplicationInfo* ParentInfo = Find(Parent))
+			{
+				ParentInfo->DependentActorList.PrepareForWrite();
+				ParentInfo->DependentActorList.Remove(Child);
+			}
+
+			if (FGlobalActorReplicationInfo* ChildInfo = Find(Child))
+			{
+				ChildInfo->ParentActorList.PrepareForWrite();
+				ChildInfo->ParentActorList.Remove(Parent);
+			}
+		}
+	}
+
 private:
 
 	TMap<FActorRepListType, TUniquePtr<FGlobalActorReplicationInfo>> ActorMap;
