Prevent RepGraph from trying to replicate invalid actors.

diff --git a/Engine/Plugins/Runtime/ReplicationGraph/Source/Private/ReplicationGraph.cpp b/Engine/Plugins/Runtime/ReplicationGraph/Source/Private/ReplicationGraph.cpp
--- a/Engine/Plugins/Runtime/ReplicationGraph/Source/Private/ReplicationGraph.cpp
+++ b/Engine/Plugins/Runtime/ReplicationGraph/Source/Private/ReplicationGraph.cpp
@@ -1481,11 +1481,29 @@ int64 UReplicationGraph::ReplicateSingleActor(AActor* Actor, FConnectionReplicat
 		UE_LOG(LogReplicationGraph, Display, TEXT("UReplicationGraph::ReplicateSingleActor: %s. NetConnection: %s"), *Actor->GetName(), *NetConnection->Describe());
 	}
 
-	if (ActorInfo.Channel && ActorInfo.Channel->Closing)
+	// These checks will happen anyway in UActorChannel::ReplicateActor, but we need to be able to detect them to prevent crashes.
+	// We could consider removing the actor from RepGraph if we hit these cases, but we don't have a good way to notify
+	// game code or the Net Driver.
+	if (!ensureMsgf(Actor, TEXT("Null Actor! Channel = %s"), *DescribeSafe(ActorInfo.Channel)))
 	{
-		// We are waiting for the client to ack this actor channel's close bunch.
 		return 0;
 	}
+	else if (!ensureMsgf(IsActorValidForReplication(Actor), TEXT("Actor not valid for replication! Actor = %s, Channel = %s"), *Actor->GetFullName(), *DescribeSafe(ActorInfo.Channel)))
+	{
+		return 0;
+	}
+	if (LIKELY(ActorInfo.Channel))
+	{
+		if (UNLIKELY(ActorInfo.Channel->Closing))
+		{
+			// We are waiting for the client to ack this actor channel's close bunch.
+			return 0;
+		}
+		else if (!ensureMsgf(ActorInfo.Channel->Actor == Actor, TEXT("Mismatched channel actors! Channel = %s, Replicating Actor = %s"), *ActorInfo.Channel->Describe(), *Actor->GetFullName()))
+		{
+			return 0;
+		}
+	}
 
 	ActorInfo.LastRepFrameNum = FrameNum;
 	ActorInfo.NextReplicationFrameNum = FrameNum + ActorInfo.ReplicationPeriodFrame;
diff --git a/Engine/Plugins/Runtime/ReplicationGraph/Source/Public/ReplicationGraphTypes.h b/Engine/Plugins/Runtime/ReplicationGraph/Source/Public/ReplicationGraphTypes.h
--- a/Engine/Plugins/Runtime/ReplicationGraph/Source/Public/ReplicationGraphTypes.h
+++ b/Engine/Plugins/Runtime/ReplicationGraph/Source/Public/ReplicationGraphTypes.h
@@ -74,7 +74,7 @@ enum class EActorRepListTypeFlags : uint8
 };
 
 // Tests if an actor is valid for replication: not pending kill, etc. Says nothing about wanting to replicate or should replicate, etc.
-FORCEINLINE bool IsActorValidForReplication(const FActorRepListType& In) { return !In->IsPendingKill() && !In->IsPendingKillPending(); }
+FORCEINLINE bool IsActorValidForReplication(const FActorRepListType& In) { return !In->IsActorBeingDestroyed() && !In->IsPendingKillOrUnreachable(); }
 
 // Tests if an actor is valid for replication gathering. Meaning, it can be gathered from the replication graph and considered for replication.
 FORCEINLINE bool IsActorValidForReplicationGather(const FActorRepListType& In)
