Added ability to limit the spatialization 2D grid to a specific region via SetBiasAndGridBounds
-When the optional grid bound is set we clamp player view locations to the closest grid in the bounded region instead of growing the grid list to his location.
-When adding actors we limit the available cells to coordinates inside the region.

diff --git a/Engine/Plugins/Runtime/ReplicationGraph/Source/Private/ReplicationGraph.cpp b/Engine/Plugins/Runtime/ReplicationGraph/Source/Private/ReplicationGraph.cpp
--- a/Engine/Plugins/Runtime/ReplicationGraph/Source/Private/ReplicationGraph.cpp
+++ b/Engine/Plugins/Runtime/ReplicationGraph/Source/Private/ReplicationGraph.cpp
@@ -817,6 +817,8 @@ int32 UReplicationGraph::ServerReplicateActors(float DeltaSeconds)
 				Node->GatherActorListsForConnection(Parameters);
 			}
 
+			Parameters.ConnectionManager.LastGatherLocation = Parameters.Viewer.ViewLocation;
+
 			if (GatheredReplicationListsForConnection.NumLists() == 0)
 			{
 				// No lists were returned, kind of weird but not fatal. Early out because code below assumes at least 1 list
@@ -1552,8 +1554,6 @@ int64 UReplicationGraph::ReplicateSingleActor(AActor* Actor, FConnectionReplicat
 	int64 BitsWritten = 0;
 	const double StartingReplicateActorTimeSeconds = GReplicateActorTimeSeconds;
 
-	ensureMsgf(ActorInfo.Channel->Actor != nullptr, TEXT("Invalid ActorChannel for %s | Channel status: pooled:%u closing:%d"), *Actor->GetName(), ActorInfo.Channel->bPooled, ActorInfo.Channel->Closing);
-					
 	if (UNLIKELY(ActorInfo.bTearOff))
 	{
 		// Replicate and immediately close in tear off case
@@ -3717,6 +3717,7 @@ static FAutoConsoleVariableRef CVarRepGraphDebugNextActor(TEXT("Net.RepGraph.Spa
 UReplicationGraphNode_GridSpatialization2D::UReplicationGraphNode_GridSpatialization2D()
 	: CellSize(0.f)
 	, SpatialBias(ForceInitToZero)
+	, GridBounds(ForceInitToZero)
 {
 	bRequiresPrepareForReplicationCall = true;
 }
@@ -3826,7 +3827,7 @@ void UReplicationGraphNode_GridSpatialization2D::AddActorInternal_Static_Impleme
 		UE_LOG(LogReplicationGraph, Display, TEXT("UReplicationGraphNode_GridSpatialization2D::AddActorInternal_Static placing %s into static grid at %s"), *Actor->GetPathName(), *ActorRepInfo.WorldLocation.ToString());
 	}
 		
-	if (SpatialBias.X > Location3D.X || SpatialBias.Y > Location3D.Y)
+	if (WillActorLocationGrowSpatialBounds(Location3D))
 	{
 		HandleActorOutOfSpatialBounds(Actor, Location3D, true);
 	}
@@ -3954,6 +3955,15 @@ void UReplicationGraphNode_GridSpatialization2D::GetGridNodesForActor(FActorRepL
 	GetGridNodesForActor(Actor, GetCellInfoForActor(Actor, ActorRepInfo.WorldLocation, ActorRepInfo.Settings.CullDistanceSquared), OutNodes);
 }
 
+void UReplicationGraphNode_GridSpatialization2D::SetBiasAndGridBounds(const FBox& GridBox)
+{
+	const FVector2D BoxMin2D = FVector2D(GridBox.Min);
+	const FVector2D BoxMax2D = FVector2D(GridBox.Max);
+	
+	SpatialBias = BoxMin2D;
+	GridBounds = FBox( FVector(BoxMin2D, -HALF_WORLD_MAX), FVector(BoxMax2D, HALF_WORLD_MAX) );
+}
+
 UReplicationGraphNode_GridSpatialization2D::FActorCellInfo UReplicationGraphNode_GridSpatialization2D::GetCellInfoForActor(FActorRepListType Actor, const FVector& Location3D, float CullDistanceSquared)
 {
 #if !(UE_BUILD_SHIPPING || UE_BUILD_TEST)
@@ -3981,8 +3991,15 @@ UReplicationGraphNode_GridSpatialization2D::FActorCellInfo UReplicationGraphNode
 	const float Dist = FMath::Sqrt(CullDistanceSquared);	 // Fixme Sqrt
 	const float MinX = LocationBiasX - Dist;
 	const float MinY = LocationBiasY - Dist;
-	const float MaxX = LocationBiasX + Dist;
-	const float MaxY = LocationBiasY + Dist;
+	float MaxX = LocationBiasX + Dist;
+	float MaxY = LocationBiasY + Dist;
+
+	if (GridBounds.IsValid)
+	{
+		const FVector BoundSize = GridBounds.GetSize();
+		MaxX = FMath::Min(MaxX, BoundSize.X);
+		MaxY = FMath::Min(MaxY, BoundSize.Y);
+	}
 
 	CellInfo.StartX = FMath::Max<int32>(0, MinX / CellSize);
 	CellInfo.StartY = FMath::Max<int32>(0, MinY / CellSize);
@@ -4039,6 +4056,12 @@ void UReplicationGraphNode_GridSpatialization2D::GetGridNodesForActor(FActorRepL
 */
 }
 
+bool UReplicationGraphNode_GridSpatialization2D::WillActorLocationGrowSpatialBounds(const FVector& Location) const
+{
+	// When bounds are set, we don't grow the cells but instead clamp the actor to the bounds.
+	return GridBounds.IsValid ? false : (SpatialBias.X > Location.X || SpatialBias.Y > Location.Y);
+}
+
 void UReplicationGraphNode_GridSpatialization2D::HandleActorOutOfSpatialBounds(AActor* Actor, const FVector& Location3D, const bool bStaticActor)
 {
 	// Don't rebuild spatialization for blacklisted actors. They will just get clamped to the grid.
@@ -4120,7 +4143,7 @@ void UReplicationGraphNode_GridSpatialization2D::PrepareForReplication()
 			const FVector Location3D = DynamicActor->GetActorLocation();
 			ActorRepInfo.WorldLocation = Location3D;
 			
-			if (SpatialBias.X > Location3D.X || SpatialBias.Y > Location3D.Y)
+			if (WillActorLocationGrowSpatialBounds(Location3D))
 			{
 				HandleActorOutOfSpatialBounds(DynamicActor, Location3D, false);
 			}
@@ -4478,12 +4501,19 @@ void UReplicationGraphNode_GridSpatialization2D::GatherActorListsForConnection(c
 		return;
 	}
 
+	FVector ClampedViewLoc = Params.Viewer.ViewLocation;
+
+	if (GridBounds.IsValid)
+	{
+		ClampedViewLoc = GridBounds.GetClosestPointTo(ClampedViewLoc);
+	}
+
 	// Find out what bucket the view is in
 
-	int32 CellX = (Params.Viewer.ViewLocation.X - SpatialBias.X) / CellSize;
+	int32 CellX = (ClampedViewLoc.X - SpatialBias.X) / CellSize;
 	if (CellX < 0)
 	{
-		UE_LOG(LogReplicationGraph, Log, TEXT("Net view location.X %s is less than the spatial bias %s"), *Params.Viewer.ViewLocation.ToString(), *SpatialBias.ToString());
+		UE_LOG(LogReplicationGraph, Log, TEXT("Net view location.X %s is less than the spatial bias %s"), *ClampedViewLoc.ToString(), *SpatialBias.ToString());
 		CellX = 0;
 	}
 	
@@ -4491,10 +4521,10 @@ void UReplicationGraphNode_GridSpatialization2D::GatherActorListsForConnection(c
 
 	// -----------
 
-	int32 CellY = (Params.Viewer.ViewLocation.Y - SpatialBias.Y) / CellSize;
+	int32 CellY = (ClampedViewLoc.Y - SpatialBias.Y) / CellSize;
 	if (CellY < 0)
 	{
-		UE_LOG(LogReplicationGraph, Log, TEXT("Net view location.Y %s is less than the spatial bias %s"), *Params.Viewer.ViewLocation.ToString(), *SpatialBias.ToString());
+		UE_LOG(LogReplicationGraph, Log, TEXT("Net view location.Y %s is less than the spatial bias %s"), *ClampedViewLoc.ToString(), *SpatialBias.ToString());
 		CellY = 0;
 	}
 	if (GridX.Num() <= CellY)
@@ -4566,8 +4596,6 @@ void UReplicationGraphNode_GridSpatialization2D::GatherActorListsForConnection(c
 			}
 		}
 	}
-
-	Params.ConnectionManager.LastGatherLocation = Params.Viewer.ViewLocation;
 }
 
 void UReplicationGraphNode_GridSpatialization2D::NotifyActorCullDistChange(AActor* Actor, FGlobalActorReplicationInfo& GlobalInfo, float OldDistSq)
diff --git a/Engine/Plugins/Runtime/ReplicationGraph/Source/Public/ReplicationGraph.h b/Engine/Plugins/Runtime/ReplicationGraph/Source/Public/ReplicationGraph.h
--- a/Engine/Plugins/Runtime/ReplicationGraph/Source/Public/ReplicationGraph.h
+++ b/Engine/Plugins/Runtime/ReplicationGraph/Source/Public/ReplicationGraph.h
@@ -498,7 +498,13 @@ public:
 	float		CellSize;
 	FVector2D	SpatialBias;
 	float		ConnectionMaxZ = WORLD_MAX; // Connection locations have to be <= to this to pull from the grid
-	
+
+	/** When the GridBounds is set we limit the creation of cells to be exclusively inside the passed region.
+	    Viewers who gather nodes outside this region will be clamped to the closest cell inside the box.
+		Actors whose location is outside the box will be clamped to the closest cell inside the box.
+	*/
+	void SetBiasAndGridBounds(const FBox& GridBox);
+
 	// Allow graph to override function for creating cell nodes in this grid.
 	TFunction<UReplicationGraphNode_GridCell*(UReplicationGraphNode_GridSpatialization2D* Parent)>	CreateCellNodeOverride;
 
@@ -522,9 +528,14 @@ protected:
 
 private:
 
+	bool WillActorLocationGrowSpatialBounds(const FVector& Location) const;
+
 	/** Called when an actor is out of spatial bounds */
 	void HandleActorOutOfSpatialBounds(AActor* Actor, const FVector& Location3D, const bool bStaticActor);
 
+	// Optional value to limit the grid to a specific region
+	FBox GridBounds;
+
 	// Classmap of actor classes which CANNOT force a rebuild of the spatialization tree. They will be clamped instead. E.g, projectiles.
 	TClassMap<bool> RebuildSpatialBlacklistMap;
 	
