Fix CIS warning V595 in RepGraph from dereferencing an actor before null check.

diff --git a/Engine/Plugins/Runtime/ReplicationGraph/Source/Private/ReplicationGraph.cpp b/Engine/Plugins/Runtime/ReplicationGraph/Source/Private/ReplicationGraph.cpp
--- a/Engine/Plugins/Runtime/ReplicationGraph/Source/Private/ReplicationGraph.cpp
+++ b/Engine/Plugins/Runtime/ReplicationGraph/Source/Private/ReplicationGraph.cpp
@@ -1476,13 +1476,6 @@ int64 UReplicationGraph::ReplicateSingleActor(AActor* Actor, FConnectionReplicat
 {
 	RG_QUICK_SCOPE_CYCLE_COUNTER(NET_ReplicateActors_ReplicateSingleActor);
 
-	UNetConnection* NetConnection = ConnectionManager.NetConnection;
-
-	if (RepGraphConditionalActorBreakpoint(Actor, NetConnection))
-	{
-		UE_LOG(LogReplicationGraph, Display, TEXT("UReplicationGraph::ReplicateSingleActor: %s. NetConnection: %s"), *Actor->GetName(), *NetConnection->Describe());
-	}
-
 	// These checks will happen anyway in UActorChannel::ReplicateActor, but we need to be able to detect them to prevent crashes.
 	// We could consider removing the actor from RepGraph if we hit these cases, but we don't have a good way to notify
 	// game code or the Net Driver.
@@ -1490,10 +1483,19 @@ int64 UReplicationGraph::ReplicateSingleActor(AActor* Actor, FConnectionReplicat
 	{
 		return 0;
 	}
-	else if (!ensureMsgf(IsActorValidForReplication(Actor), TEXT("Actor not valid for replication! Actor = %s, Channel = %s"), *Actor->GetFullName(), *DescribeSafe(ActorInfo.Channel)))
+
+	UNetConnection* NetConnection = ConnectionManager.NetConnection;
+
+	if (RepGraphConditionalActorBreakpoint(Actor, NetConnection))
+	{
+		UE_LOG(LogReplicationGraph, Display, TEXT("UReplicationGraph::ReplicateSingleActor: %s. NetConnection: %s"), *Actor->GetName(), *NetConnection->Describe());
+	}
+
+	if (!ensureMsgf(IsActorValidForReplication(Actor), TEXT("Actor not valid for replication! Actor = %s, Channel = %s"), *Actor->GetFullName(), *DescribeSafe(ActorInfo.Channel)))
 	{
 		return 0;
 	}
+
 	if (LIKELY(ActorInfo.Channel))
 	{
 		if (UNLIKELY(ActorInfo.Channel->Closing))
