-The RepGraph forced a garbage collection when rebuilding a GridSpace. This will now be done at the end of frame instead of in the tick.
-This forced GC will only trigger if more than 100 cells were destroyed.

diff --git a/Engine/Plugins/Runtime/ReplicationGraph/Source/Private/ReplicationGraph.cpp b/Engine/Plugins/Runtime/ReplicationGraph/Source/Private/ReplicationGraph.cpp
--- a/Engine/Plugins/Runtime/ReplicationGraph/Source/Private/ReplicationGraph.cpp
+++ b/Engine/Plugins/Runtime/ReplicationGraph/Source/Private/ReplicationGraph.cpp
@@ -86,6 +86,9 @@ static FAutoConsoleVariableRef CVarRepGraphPrintCulledOnConnectionClasses(TEXT("
 int32 CVar_RepGraph_TrackClassReplication = 0;
 static FAutoConsoleVariableRef CVarRepGraphTrackClassReplication(TEXT("Net.RepGraph.TrackClassReplication"), CVar_RepGraph_TrackClassReplication, TEXT(""), ECVF_Default );
 
+int32 CVar_RepGraph_NbDestroyedGridsToTriggerGC = 100;
+static FAutoConsoleVariableRef CVarRepGraphNbDestroyedGridsToTriggerGC(TEXT("Net.RepGraph.NbDestroyedGridsToTriggerGC"), CVar_RepGraph_NbDestroyedGridsToTriggerGC, TEXT(""), ECVF_Default);
+
 int32 CVar_RepGraph_PrintTrackClassReplication = 0;
 static FAutoConsoleVariableRef CVarRepGraphPrintTrackClassReplication(TEXT("Net.RepGraph.PrintTrackClassReplication"), CVar_RepGraph_PrintTrackClassReplication, TEXT(""), ECVF_Default );
 
@@ -4438,6 +4441,7 @@ void UReplicationGraphNode_GridSpatialization2D::PrepareForReplication()
 		UE_LOG(LogReplicationGraph, Warning, TEXT("Rebuilding spatialization graph for bias %s"), *SpatialBias.ToString());
 		
 		// Tear down all existing nodes first. This marks them pending kill.
+		int32 GridsDestroyed(0);
 		for (auto& InnerArray : Grid)
 		{
 			for (UReplicationGraphNode_GridCell*& N : InnerArray)
@@ -4446,14 +4450,17 @@ void UReplicationGraphNode_GridSpatialization2D::PrepareForReplication()
 				{
 					N->TearDown();
 					N = nullptr;
+					++GridsDestroyed;
 				}
 			}
 		}
 
 		// Force a garbage collection. Without this you may hit OOMs if rebuilding spatialization every frame for some period of time. 
 		// (Obviously not ideal to ever be doing this. But you are already hitching, might as well GC to avoid OOM crash).
-		
-		CollectGarbage( GARBAGE_COLLECTION_KEEPFLAGS, true );
+		if (GridsDestroyed >= CVar_RepGraph_NbDestroyedGridsToTriggerGC)
+		{
+			GEngine->ForceGarbageCollection(true);
+		}
 		
 		for (auto& MapIt : DynamicSpatializedActors)
 		{
