Keep the access specifiers of UNetReplicationGraphConnection virtual functions inherited from UReplicationConnectionDriver public.

diff --git a/Engine/Plugins/Runtime/ReplicationGraph/Source/Public/ReplicationGraph.h b/Engine/Plugins/Runtime/ReplicationGraph/Source/Public/ReplicationGraph.h
--- a/Engine/Plugins/Runtime/ReplicationGraph/Source/Public/ReplicationGraph.h
+++ b/Engine/Plugins/Runtime/ReplicationGraph/Source/Public/ReplicationGraph.h
@@ -994,9 +994,9 @@ public:
 	//~ Begin UReplicationConnectionDriver Interface
 	virtual void TearDown() override;
 
-private:
+	virtual void NotifyClientVisibleLevelNamesAdd(FName LevelName, UWorld* StreamingWorld) override;
 
-	friend UReplicationGraph;
+	virtual void NotifyClientVisibleLevelNamesRemove(FName LevelName) override { OnClientVisibleLevelNameRemove.Broadcast(LevelName); }
 
 	virtual void NotifyActorChannelAdded(AActor* Actor, class UActorChannel* Channel) override;
 
@@ -1009,11 +1009,11 @@ private:
 	virtual void NotifyRemoveDestructionInfo(FActorDestructionInfo* DestructInfo) override;
 
 	virtual void NotifyResetDestructionInfo() override;
+	//~ End UReplicationConnectionDriver Interface
 
-	virtual void NotifyClientVisibleLevelNamesAdd(FName LevelName, UWorld* StreamingWorld) override;
+private:
 
-	virtual void NotifyClientVisibleLevelNamesRemove(FName LevelName) override { OnClientVisibleLevelNameRemove.Broadcast(LevelName); }
-	//~ End UReplicationConnectionDriver Interface
+	friend UReplicationGraph;
 
 	// ----------------------------------------
 
