-Now caching CullDistance and CullDistanceSquared in ConnectionReplicationActorInfo
-Removes sqrt call in critical path when RepGraph does actor culling.

diff --git a/Engine/Plugins/Runtime/ReplicationGraph/Source/Private/BasicReplicationGraph.cpp b/Engine/Plugins/Runtime/ReplicationGraph/Source/Private/BasicReplicationGraph.cpp
--- a/Engine/Plugins/Runtime/ReplicationGraph/Source/Private/BasicReplicationGraph.cpp
+++ b/Engine/Plugins/Runtime/ReplicationGraph/Source/Private/BasicReplicationGraph.cpp
@@ -41,11 +41,11 @@ void UBasicReplicationGraph::InitGlobalActorClassSettings()
 		
 		if (ActorCDO->bAlwaysRelevant || ActorCDO->bOnlyRelevantToOwner)
 		{
-			ClassInfo.CullDistanceSquared = 0.f;
+			ClassInfo.SetCullDistanceSquared(0.f);
 		}
 		else
 		{
-			ClassInfo.CullDistanceSquared = ActorCDO->NetCullDistanceSquared;
+			ClassInfo.SetCullDistanceSquared(ActorCDO->NetCullDistanceSquared);
 		}
 		
 		GlobalActorReplicationInfoMap.SetClassInfo( Class, ClassInfo );
diff --git a/Engine/Plugins/Runtime/ReplicationGraph/Source/Private/ReplicationGraph.cpp b/Engine/Plugins/Runtime/ReplicationGraph/Source/Private/ReplicationGraph.cpp
--- a/Engine/Plugins/Runtime/ReplicationGraph/Source/Private/ReplicationGraph.cpp
+++ b/Engine/Plugins/Runtime/ReplicationGraph/Source/Private/ReplicationGraph.cpp
@@ -1057,7 +1057,7 @@ void UReplicationGraph::ReplicateActorListsForConnection_Default(UNetReplication
 				{
 					const float DistSq = (GlobalData.WorldLocation - ConnectionViewLocation).SizeSquared();
 
-					if (bDoDistanceCull && ConnectionData.CullDistanceSquared > 0.f && DistSq > ConnectionData.CullDistanceSquared)
+					if (bDoDistanceCull && ConnectionData.GetCullDistanceSquared() > 0.f && DistSq > ConnectionData.GetCullDistanceSquared())
 					{
 						DO_REPGRAPH_DETAILS(PrioritizedReplicationList.GetNextSkippedDebugDetails(Actor)->DistanceCulled = FMath::Sqrt(DistSq));
 						if (bDoCulledOnConnectionCount)
@@ -1333,7 +1333,7 @@ void UReplicationGraph::ReplicateActorListsForConnection_FastShared(UNetReplicat
 
 			// Simple distance cull
 			const float DistSq = DirToActor.SizeSquared();
-			if (DistSq > (ConnectionData.CullDistanceSquared * FastSharedDistanceRequirementPct))
+			if (DistSq > (ConnectionData.GetCullDistanceSquared() * FastSharedDistanceRequirementPct))
 			{
 				continue;
 			}
@@ -1769,7 +1769,7 @@ bool UReplicationGraph::ProcessRemoteFunction(class AActor* Actor, UFunction* Fu
 
 		RepLayout->BuildSharedSerializationForRPC(Parameters);
 		FGlobalActorReplicationInfo& GlobalInfo = GlobalActorReplicationInfoMap.Get(Actor);
-		const float CullDistanceSquared = GlobalInfo.Settings.CullDistanceSquared;
+		const float CullDistanceSquared = GlobalInfo.Settings.GetCullDistanceSquared();
 
 		bool ForceFlushNetDormancy = false;
 
@@ -3086,7 +3086,7 @@ FORCEINLINE void UReplicationGraphNode_DynamicSpatialFrequency::CalcFrequencyFor
 	// Skip if past cull distance
 	const FVector DirToActor = GlobalInfo.WorldLocation - ConnectionViewLocation;
 	const float DistanceToActorSq = DirToActor.SizeSquared();
-	if (!IgnoreCullDistance && ConnectionInfo.CullDistanceSquared > 0.f && DistanceToActorSq > ConnectionInfo.CullDistanceSquared)
+	if (!IgnoreCullDistance && ConnectionInfo.GetCullDistanceSquared() > 0.f && DistanceToActorSq > ConnectionInfo.GetCullDistanceSquared())
 	{
 		RemoveExistingItem();
 		return;
@@ -3115,7 +3115,7 @@ FORCEINLINE void UReplicationGraphNode_DynamicSpatialFrequency::CalcFrequencyFor
 			// --------------------------------------------------------------------------------------------------------
 			{
 				// Calc Percentage of distance relative to cull distance, scaled to ZoneInfo Min/Max pct
-				const float CullDistSq = ConnectionInfo.CullDistanceSquared > 0.f ? ConnectionInfo.CullDistanceSquared : GlobalInfo.Settings.CullDistanceSquared; // Use global settings if the connection specific setting is zero'd out
+				const float CullDistSq = ConnectionInfo.GetCullDistanceSquared() > 0.f ? ConnectionInfo.GetCullDistanceSquared() : GlobalInfo.Settings.GetCullDistanceSquared(); // Use global settings if the connection specific setting is zero'd out
 
 				if (!ensureMsgf(CullDistSq > 0.f, TEXT("UReplicationGraphNode_DynamicSpatialFrequency::GatherActors: %s has cull distance of 0. Skipping"), *GetPathNameSafe(Actor)))
 				{
@@ -3124,7 +3124,7 @@ FORCEINLINE void UReplicationGraphNode_DynamicSpatialFrequency::CalcFrequencyFor
 					return;
 				}
 
-				const float CullDist = FMath::Sqrt(CullDistSq); // Fixme: sqrt
+				const float CullDist = ConnectionInfo.GetCullDistance();
 				const float DistPct =  DistanceToActor / CullDist;
 
 				const float BiasDistPct = DistPct - ZoneInfo.MinDistPct;
@@ -3305,10 +3305,10 @@ void UReplicationGraphNode_ConnectionDormanyNode::ConditionalGatherDormantActors
 		if (ConnectionActorInfo.bDormantOnConnection)
 		{
 			// If we trickled this actor, restore CullDistance to the default
-			if (ConnectionActorInfo.CullDistanceSquared <= 0.f)
+			if (ConnectionActorInfo.GetCullDistanceSquared() <= 0.f)
 			{
 				FGlobalActorReplicationInfo& GlobalInfo = GlobalActorReplicationInfoMap->Get(Actor);
-				ConnectionActorInfo.CullDistanceSquared = GlobalInfo.Settings.CullDistanceSquared;
+				ConnectionActorInfo.SetCullDistanceSquared(GlobalInfo.Settings.GetCullDistanceSquared());
 			}
 
 			// He can be removed
@@ -3324,7 +3324,7 @@ void UReplicationGraphNode_ConnectionDormanyNode::ConditionalGatherDormantActors
 		}
 		else if (CVar_RepGraph_TrickleDistCullOnDormanyNodes > 0 && bShouldTrickle)
 		{
-			ConnectionActorInfo.CullDistanceSquared = 0.f;
+			ConnectionActorInfo.SetCullDistanceSquared(0.f);
 			bShouldTrickle = false; // trickle one actor per frame
 		}
 	}
@@ -3957,7 +3957,7 @@ void UReplicationGraphNode_GridSpatialization2D::PutStaticActorIntoCell(const FN
 void UReplicationGraphNode_GridSpatialization2D::GetGridNodesForActor(FActorRepListType Actor, const FGlobalActorReplicationInfo& ActorRepInfo, TArray<UReplicationGraphNode_GridCell*>& OutNodes)
 {
 	RG_QUICK_SCOPE_CYCLE_COUNTER(UReplicationGraphNode_GridSpatialization2D_GetGridNodesForActor);
-	GetGridNodesForActor(Actor, GetCellInfoForActor(Actor, ActorRepInfo.WorldLocation, ActorRepInfo.Settings.CullDistanceSquared), OutNodes);
+	GetGridNodesForActor(Actor, GetCellInfoForActor(Actor, ActorRepInfo.WorldLocation, ActorRepInfo.Settings.GetCullDistance()), OutNodes);
 }
 
 void UReplicationGraphNode_GridSpatialization2D::SetBiasAndGridBounds(const FBox& GridBox)
@@ -3969,12 +3969,12 @@ void UReplicationGraphNode_GridSpatialization2D::SetBiasAndGridBounds(const FBox
 	GridBounds = FBox( FVector(BoxMin2D, -HALF_WORLD_MAX), FVector(BoxMax2D, HALF_WORLD_MAX) );
 }
 
-UReplicationGraphNode_GridSpatialization2D::FActorCellInfo UReplicationGraphNode_GridSpatialization2D::GetCellInfoForActor(FActorRepListType Actor, const FVector& Location3D, float CullDistanceSquared)
+UReplicationGraphNode_GridSpatialization2D::FActorCellInfo UReplicationGraphNode_GridSpatialization2D::GetCellInfoForActor(FActorRepListType Actor, const FVector& Location3D, float CullDistance)
 {
 #if !(UE_BUILD_SHIPPING || UE_BUILD_TEST)
-	if (CullDistanceSquared <= 0.f)
+	if (CullDistance <= 0.f)
 	{
-		UE_LOG(LogReplicationGraph, Warning, TEXT("::GetGridNodesForActor called on %s when its CullDistanceSquared = %.2f. (Must be > 0)"), *GetActorRepListTypeDebugString(Actor), CullDistanceSquared);
+		UE_LOG(LogReplicationGraph, Warning, TEXT("::GetGridNodesForActor called on %s when its CullDistance = %.2f. (Must be > 0)"), *GetActorRepListTypeDebugString(Actor), CullDistance);
 	}
 #endif
 
@@ -3993,7 +3993,7 @@ UReplicationGraphNode_GridSpatialization2D::FActorCellInfo UReplicationGraphNode
 	const float LocationBiasX = (ClampedLocation.X - SpatialBias.X);
 	const float LocationBiasY = (ClampedLocation.Y - SpatialBias.Y);
 
-	const float Dist = FMath::Sqrt(CullDistanceSquared);	 // Fixme Sqrt
+	const float Dist = CullDistance;
 	const float MinX = LocationBiasX - Dist;
 	const float MinY = LocationBiasY - Dist;
 	float MaxX = LocationBiasX + Dist;
@@ -4156,7 +4156,7 @@ void UReplicationGraphNode_GridSpatialization2D::PrepareForReplication()
 			if (!bNeedsRebuild)
 			{
 				// Get the new CellInfo
-				const FActorCellInfo NewCellInfo = GetCellInfoForActor(DynamicActor, Location3D, ActorRepInfo.Settings.CullDistanceSquared);
+				const FActorCellInfo NewCellInfo = GetCellInfoForActor(DynamicActor, Location3D, ActorRepInfo.Settings.GetCullDistance());
 
 				if (PreviousCellInfo.IsValid())
 				{
@@ -4476,7 +4476,7 @@ void UReplicationGraphNode_GridSpatialization2D::PrepareForReplication()
 				FGlobalActorReplicationInfo& ActorRepInfo = GlobalRepMap->Get(DynamicActor);
 				ActorRepInfo.WorldLocation = Location3D;
 
-				const FActorCellInfo NewCellInfo = GetCellInfoForActor(DynamicActor, Location3D, ActorRepInfo.Settings.CullDistanceSquared);
+				const FActorCellInfo NewCellInfo = GetCellInfoForActor(DynamicActor, Location3D, ActorRepInfo.Settings.GetCullDistance());
 
 				GetGridNodesForActor(DynamicActor, NewCellInfo, GatheredNodes);
 				for (UReplicationGraphNode_GridCell* Node : GatheredNodes)
@@ -4589,7 +4589,7 @@ void UReplicationGraphNode_GridSpatialization2D::GatherActorListsForConnection(c
 						ActorInfo->bDormantOnConnection = false;
 
 						// add back to connection specific dormancy nodes
-						const FActorCellInfo CellInfo = GetCellInfoForActor(Actor, Actor->GetActorLocation(), ActorInfo->CullDistanceSquared);
+						const FActorCellInfo CellInfo = GetCellInfoForActor(Actor, Actor->GetActorLocation(), ActorInfo->GetCullDistance());
 
 						GetGridNodesForActor(Actor, CellInfo, GatheredNodes);
 
@@ -4607,7 +4607,7 @@ void UReplicationGraphNode_GridSpatialization2D::GatherActorListsForConnection(c
 	}
 }
 
-void UReplicationGraphNode_GridSpatialization2D::NotifyActorCullDistChange(AActor* Actor, FGlobalActorReplicationInfo& GlobalInfo, float OldDistSq)
+void UReplicationGraphNode_GridSpatialization2D::NotifyActorCullDistChange(AActor* Actor, FGlobalActorReplicationInfo& GlobalInfo, float OldDist)
 {
 	RG_QUICK_SCOPE_CYCLE_COUNTER(UReplicationGraphNode_GridSpatialization2D_NotifyActorCullDistChange);
 
@@ -4615,14 +4615,14 @@ void UReplicationGraphNode_GridSpatialization2D::NotifyActorCullDistChange(AActo
 	if (FCachedStaticActorInfo* StaticActorInfo = StaticSpatializedActors.Find(Actor))
 	{
 		// Remove with old distance
-		GetGridNodesForActor(Actor, GetCellInfoForActor(Actor, GlobalInfo.WorldLocation, OldDistSq), GatheredNodes);
+		GetGridNodesForActor(Actor, GetCellInfoForActor(Actor, GlobalInfo.WorldLocation, OldDist), GatheredNodes);
 		for (UReplicationGraphNode_GridCell* Node : GatheredNodes)
 		{
 			Node->RemoveStaticActor(StaticActorInfo->ActorInfo, GlobalInfo, GlobalInfo.bWantsToBeDormant);
 		}
 
 		// Add new distances (there is some waste here but this hopefully doesn't happen much at runtime!)
-		GetGridNodesForActor(Actor, GetCellInfoForActor(Actor, GlobalInfo.WorldLocation, GlobalInfo.Settings.CullDistanceSquared), GatheredNodes);
+		GetGridNodesForActor(Actor, GetCellInfoForActor(Actor, GlobalInfo.WorldLocation, GlobalInfo.Settings.GetCullDistance()), GatheredNodes);
 		for (UReplicationGraphNode_GridCell* Node : GatheredNodes)
 		{
 			Node->AddStaticActor(StaticActorInfo->ActorInfo, GlobalInfo, StaticActorInfo->bDormancyDriven);
@@ -4650,7 +4650,7 @@ void UReplicationGraphNode_GridSpatialization2D::NotifyActorCullDistChange(AActo
 		// Might be in the pending init list
 		if (PendingStaticSpatializedActors.FindByKey(Actor) == nullptr)
 		{
-			UE_LOG(LogReplicationGraph, Warning, TEXT("UReplicationGraphNode_GridSpatialization2D::NotifyActorCullDistChange. %s Changed Cull Distance (%.2f -> %.2f) but is not in static or dynamic actor lists. %s"), *Actor->GetPathName(), FMath::Sqrt(OldDistSq), FMath::Sqrt(GlobalInfo.Settings.CullDistanceSquared), *GetPathName() );
+			UE_LOG(LogReplicationGraph, Warning, TEXT("UReplicationGraphNode_GridSpatialization2D::NotifyActorCullDistChange. %s Changed Cull Distance (%.2f -> %.2f) but is not in static or dynamic actor lists. %s"), *Actor->GetPathName(), OldDist, GlobalInfo.Settings.GetCullDistance(), *GetPathName() );
 
 			// Search the entire grid. This is slow so only enabled if verify is on.
 			if (CVar_RepGraph_Verify)
@@ -4802,13 +4802,13 @@ void UReplicationGraphNode_AlwaysRelevant_ForConnection::GatherActorListsForConn
 			if (NewActor)
 			{
 				// Zero out new actor cull distance
-				Params.ConnectionManager.ActorInfoMap.FindOrAdd(NewActor).CullDistanceSquared = 0.f;
+				Params.ConnectionManager.ActorInfoMap.FindOrAdd(NewActor).SetCullDistanceSquared(0.f);
 			}
 			if (LastActor)
 			{
 				// Reset previous actor culldistance
 				FConnectionReplicationActorInfo& ActorInfo = Params.ConnectionManager.ActorInfoMap.FindOrAdd(LastActor);
-				ActorInfo.CullDistanceSquared = GraphGlobals->GlobalActorReplicationInfoMap->Get(LastActor).Settings.CullDistanceSquared;
+				ActorInfo.SetCullDistanceSquared(GraphGlobals->GlobalActorReplicationInfoMap->Get(LastActor).Settings.GetCullDistanceSquared());
 			}
 
 			LastActor = NewActor;
diff --git a/Engine/Plugins/Runtime/ReplicationGraph/Source/Private/ReplicationGraphDebugging.cpp b/Engine/Plugins/Runtime/ReplicationGraph/Source/Private/ReplicationGraphDebugging.cpp
--- a/Engine/Plugins/Runtime/ReplicationGraph/Source/Private/ReplicationGraphDebugging.cpp
+++ b/Engine/Plugins/Runtime/ReplicationGraph/Source/Private/ReplicationGraphDebugging.cpp
@@ -185,7 +185,7 @@ void AReplicationGraphDebugActor::PrintCullDistances()
 		bool bFound = false;
 		for (FData& ExistingData : DataList)
 		{
-			if (ExistingData.Class == Actor->GetClass() && FMath::IsNearlyZero(ExistingData.DistSq - Info.Settings.CullDistanceSquared))
+			if (ExistingData.Class == Actor->GetClass() && FMath::IsNearlyZero(ExistingData.DistSq - Info.Settings.GetCullDistanceSquared()))
 			{
 				ExistingData.Count++;
 				bFound = true;
@@ -200,7 +200,7 @@ void AReplicationGraphDebugActor::PrintCullDistances()
 
 		FData NewData;
 		NewData.Class = Actor->GetClass();
-		NewData.DistSq = Info.Settings.CullDistanceSquared;
+		NewData.DistSq = Info.Settings.GetCullDistanceSquared();
 		NewData.Count = 1;
 		DataList.Add(NewData);
 	}
@@ -421,7 +421,7 @@ void AReplicationGraphDebugActor::ServerSetCullDistanceForClass_Implementation(U
 	const float CullDistSq = CullDistance * CullDistance;
 
 	FClassReplicationInfo& ClassInfo = ReplicationGraph->GlobalActorReplicationInfoMap.GetClassInfo(Class);
-	ClassInfo.CullDistanceSquared = CullDistSq;
+	ClassInfo.SetCullDistanceSquared(CullDistSq);
 	UE_LOG(LogReplicationGraph, Display, TEXT("Setting cull distance for class %s to %.2f"), *Class->GetName(), CullDistance);
 
 	for (TActorIterator<AActor> ActorIt(GetWorld(), Class); ActorIt; ++ActorIt)
@@ -429,14 +429,14 @@ void AReplicationGraphDebugActor::ServerSetCullDistanceForClass_Implementation(U
 		AActor* Actor = *ActorIt;
 		if (FGlobalActorReplicationInfo* ActorInfo = ReplicationGraph->GlobalActorReplicationInfoMap.Find(Actor))
 		{
-			ActorInfo->Settings.CullDistanceSquared = CullDistSq;
+			ActorInfo->Settings.SetCullDistanceSquared(CullDistSq);
 			UE_LOG(LogReplicationGraph, Display, TEXT("Setting GlobalActorInfo cull distance for %s to %.2f"), *Actor->GetName(), CullDistance);
 		}
 
 
 		if (FConnectionReplicationActorInfo* ConnectionActorInfo = ConnectionManager->ActorInfoMap.Find(Actor))
 		{
-			ConnectionActorInfo->CullDistanceSquared = CullDistSq;
+			ConnectionActorInfo->SetCullDistanceSquared(CullDistSq);
 			UE_LOG(LogReplicationGraph, Display, TEXT("Setting Connection cull distance for %s to %.2f"), *Actor->GetName(), CullDistance);
 		}
 	}
diff --git a/Engine/Plugins/Runtime/ReplicationGraph/Source/Public/ReplicationGraph.h b/Engine/Plugins/Runtime/ReplicationGraph/Source/Public/ReplicationGraph.h
--- a/Engine/Plugins/Runtime/ReplicationGraph/Source/Public/ReplicationGraph.h
+++ b/Engine/Plugins/Runtime/ReplicationGraph/Source/Public/ReplicationGraph.h
@@ -493,7 +493,7 @@ public:
 	
 
 	// Called if cull distance changes. Note the caller must update Global/Connection actor rep infos. This just changes cached state within this node
-	void NotifyActorCullDistChange(AActor* Actor, FGlobalActorReplicationInfo& GlobalInfo, float OldDistSq);
+	void NotifyActorCullDistChange(AActor* Actor, FGlobalActorReplicationInfo& GlobalInfo, float OldDist);
 	
 	float		CellSize;
 	FVector2D	SpatialBias;
@@ -624,7 +624,7 @@ private:
 	void GetGridNodesForActor(FActorRepListType Actor, const FGlobalActorReplicationInfo& ActorRepInfo, TArray<UReplicationGraphNode_GridCell*>& OutNodes);
 	void GetGridNodesForActor(FActorRepListType Actor, const UReplicationGraphNode_GridSpatialization2D::FActorCellInfo& CellInfo, TArray<UReplicationGraphNode_GridCell*>& OutNodes);
 
-	FActorCellInfo GetCellInfoForActor(FActorRepListType Actor, const FVector& Location3D, float CullDistanceSquared);
+	FActorCellInfo GetCellInfoForActor(FActorRepListType Actor, const FVector& Location3D, float CullDistance);
 
 	// This is a reused TArray for gathering actor nodes. Just to prevent using a stack based TArray everywhere or static/reset patten.
 	TArray<UReplicationGraphNode_GridCell*> GatheredNodes;
diff --git a/Engine/Plugins/Runtime/ReplicationGraph/Source/Public/ReplicationGraphTypes.h b/Engine/Plugins/Runtime/ReplicationGraph/Source/Public/ReplicationGraphTypes.h
--- a/Engine/Plugins/Runtime/ReplicationGraph/Source/Public/ReplicationGraphTypes.h
+++ b/Engine/Plugins/Runtime/ReplicationGraph/Source/Public/ReplicationGraphTypes.h
@@ -357,7 +357,6 @@ struct FClassReplicationInfo
 	FClassReplicationInfo() { }
 	float DistancePriorityScale = 1.f;
 	float StarvationPriorityScale = 1.f;
-	float CullDistanceSquared = 0.f;
 	float AccumulatedNetPriorityBias = 0.f;
 	
 	uint8 ReplicationPeriodFrame = 1;
@@ -367,13 +366,22 @@ struct FClassReplicationInfo
 	TFunction<bool(AActor*)> FastSharedReplicationFunc = nullptr;
 	FName FastSharedReplicationFuncName = NAME_None;
 
+	void SetCullDistanceSquared(float InCullDistanceSquared)
+	{
+		CullDistanceSquared = InCullDistanceSquared;
+		CullDistance = FMath::Sqrt(CullDistanceSquared);
+	}
+
+	float GetCullDistance() const { return CullDistance; }
+	float GetCullDistanceSquared() const { return CullDistanceSquared; }
+
 	FString BuildDebugStringDelta() const
 	{
 		FClassReplicationInfo DefaultValues;
 		FString Str;
-		if (CullDistanceSquared != DefaultValues.CullDistanceSquared)
+		if (CullDistance != DefaultValues.CullDistance)
 		{
-			Str += FString::Printf(TEXT("CullDistance: %.2f "), FMath::Sqrt(CullDistanceSquared));
+			Str += FString::Printf(TEXT("CullDistance: %.2f "), CullDistance);
 		}
 		if (StarvationPriorityScale != DefaultValues.StarvationPriorityScale)
 		{
@@ -402,6 +410,12 @@ struct FClassReplicationInfo
 
 		return Str;
 	}
+
+private:
+
+	float CullDistance = 0.0f;
+	float CullDistanceSquared = 0.f;
+
 };
 
 struct FGlobalActorReplicationInfo;
@@ -776,7 +790,7 @@ struct FConnectionReplicationActorInfo
 
 		ReplicationPeriodFrame = GlobalInfo.Settings.ReplicationPeriodFrame;
 		FastPath_ReplicationPeriodFrame = GlobalInfo.Settings.FastPath_ReplicationPeriodFrame;
-		CullDistanceSquared = GlobalInfo.Settings.CullDistanceSquared;
+		SetCullDistanceSquared(GlobalInfo.Settings.GetCullDistanceSquared());
 	}
 
 	/** Resets the data, except for the "settings" data that we pulled from GlobalInfo */
@@ -793,9 +807,22 @@ struct FConnectionReplicationActorInfo
 		// Note: purposefully not clearing bDormantOnConnection or bTearOff.
 	}
 
+	void SetCullDistanceSquared(float InCullDistanceSquared)
+	{
+		CullDistanceSquared = InCullDistanceSquared;
+		CullDistance = FMath::Sqrt(CullDistanceSquared);
+	}
+
+	float GetCullDistance() const { return CullDistance; }
+	float GetCullDistanceSquared() const { return CullDistanceSquared; }
+
 	UActorChannel* Channel = nullptr;
 
+private:
+	float CullDistance = 0.f;
 	float CullDistanceSquared = 0.f;
+
+public:
 	
 	/** Default replication */
 	uint32	NextReplicationFrameNum = 0;	/** The next frame we are allowed to replicate on */
