Prevent dependent actors from being added if either they or their parent are invalid.

diff --git a/Engine/Plugins/Runtime/ReplicationGraph/Source/Public/ReplicationGraphTypes.h b/Engine/Plugins/Runtime/ReplicationGraph/Source/Public/ReplicationGraphTypes.h
--- a/Engine/Plugins/Runtime/ReplicationGraph/Source/Public/ReplicationGraphTypes.h
+++ b/Engine/Plugins/Runtime/ReplicationGraph/Source/Public/ReplicationGraphTypes.h
@@ -738,7 +738,13 @@ struct FGlobalActorReplicationInfoMap
 
 	void AddDependentActor(AActor* Parent, AActor* Child)
 	{
-		if (Parent && Child)
+		const bool bIsParentValid = ensureMsgf(Parent && IsActorValidForReplication(Parent), TEXT("FGlobalActorReplicationInfoMap::AddDependentActor Invalid Parent! %s"),
+			*GetPathNameSafe(Parent));
+
+		const bool bIsChildValid = ensureMsgf(Child && IsActorValidForReplication(Child), TEXT("FGlobalActorReplicationInfoMap::AddDependentActor Invalid Child! %s"),
+			*GetPathNameSafe(Child));
+
+		if (bIsParentValid && bIsChildValid)
 		{
 			if (FGlobalActorReplicationInfo* ParentInfo = Find(Parent))
 			{
