Optimized the DormantDynamicActorsDestruction feature by not creating ConnectionDormancyNodes for cells the client never visited.

diff --git a/Engine/Plugins/Runtime/ReplicationGraph/Source/Private/ReplicationGraph.cpp b/Engine/Plugins/Runtime/ReplicationGraph/Source/Private/ReplicationGraph.cpp
--- a/Engine/Plugins/Runtime/ReplicationGraph/Source/Private/ReplicationGraph.cpp
+++ b/Engine/Plugins/Runtime/ReplicationGraph/Source/Private/ReplicationGraph.cpp
@@ -3505,6 +3505,12 @@ void UReplicationGraphNode_DormancyNode::GatherActorListsForConnection(const FCo
 	ConnectionNode->GatherActorListsForConnection(Params);
 }
 
+UReplicationGraphNode_ConnectionDormanyNode* UReplicationGraphNode_DormancyNode::GetExistingConnectionNode(const FConnectionGatherActorListParameters& Params)
+{
+	UReplicationGraphNode_ConnectionDormanyNode** ConnectionNodeItem = ConnectionNodes.Find(&Params.ConnectionManager);
+	return ConnectionNodeItem == nullptr ? nullptr : *ConnectionNodeItem;
+}
+
 UReplicationGraphNode_ConnectionDormanyNode* UReplicationGraphNode_DormancyNode::GetConnectionNode(const FConnectionGatherActorListParameters& Params)
 {
 	UReplicationGraphNode_ConnectionDormanyNode** NodePtrPtr = ConnectionNodes.Find(&Params.ConnectionManager);
@@ -4547,59 +4553,77 @@ void UReplicationGraphNode_GridSpatialization2D::GatherActorListsForConnection(c
 		CellNode->GatherActorListsForConnection(Params);
 	}
 
-	if (CVar_RepGraph_DormantDynamicActorsDestruction > 0)
+	if (bDestroyDormantDynamicActors && CVar_RepGraph_DormantDynamicActorsDestruction > 0 )
 	{
-		int32 PrevX = FMath::Max(0, (int32)((Params.ConnectionManager.LastGatherLocation.X - SpatialBias.X) / CellSize));
-		int32 PrevY = FMath::Max(0, (int32)((Params.ConnectionManager.LastGatherLocation.Y - SpatialBias.Y) / CellSize));
+		FVector LastGatherLoc = Params.ConnectionManager.LastGatherLocation;
+		
+		if (GridBounds.IsValid)
+		{
+			LastGatherLoc = GridBounds.GetClosestPointTo(LastGatherLoc);
+		}
+
+		int32 PrevX = FMath::Max(0, (int32)((LastGatherLoc.X - SpatialBias.X) / CellSize));
+		int32 PrevY = FMath::Max(0, (int32)((LastGatherLoc.Y - SpatialBias.Y) / CellSize));
 
 		// if the grid cell changed this gather
 		if ((CellX != PrevX) || (CellY != PrevY))
 		{
-			RG_QUICK_SCOPE_CYCLE_COUNTER(UReplicationGraphNode_GridSpatialization2D_CellChangeDormantRelevancy);
+			UReplicationGraphNode_GridCell* PrevCell = GetCell(PrevX, PrevY);
+			DoDormantDynamicActorDestruction(Params, CellNode, PrevCell);
+		}
+	}
+}
 
-			FActorRepListRefView DormantActorList;
-			FActorRepListRefView PrevDormantActorList;
+void UReplicationGraphNode_GridSpatialization2D::DoDormantDynamicActorDestruction(const FConnectionGatherActorListParameters& Params, UReplicationGraphNode_GridCell* CurrentCell, UReplicationGraphNode_GridCell* PreviousCell)
+{
+	RG_QUICK_SCOPE_CYCLE_COUNTER(UReplicationGraphNode_GridSpatialization2D_DoDormantDynamicActorDestruction);
 
-			if (CellNode)
-			{
-				if (UReplicationGraphNode_DormancyNode* DormancyNode = CellNode->GetDormancyNode())
-				{
-					DormancyNode->ConditionalGatherDormantDynamicActors(DormantActorList, Params, nullptr);
-				}
-			}
+	FActorRepListRefView DormantActorList;
+	FActorRepListRefView PrevDormantActorList;
 
-			TArray<UReplicationGraphNode_GridCell*>& PrevGridX = GetGridX(PrevX);
-			if (UReplicationGraphNode_GridCell* PrevCell = GetCell(PrevGridX, PrevY))
-			{
-				if (UReplicationGraphNode_DormancyNode* DormancyNode = PrevCell->GetDormancyNode())
-				{
-					DormancyNode->ConditionalGatherDormantDynamicActors(PrevDormantActorList, Params, &DormantActorList);
-				}
-			}
+	if (CurrentCell)
+	{
+		if (UReplicationGraphNode_DormancyNode* DormancyNode = CurrentCell->GetDormancyNode())
+		{
+			DormancyNode->ConditionalGatherDormantDynamicActors(DormantActorList, Params, nullptr);
+		}
+	}
 
-			// any previous dormant actors not in the current node dormant list
-			if (PrevDormantActorList.IsValid())
-			{
-				for (FActorRepListType& Actor : PrevDormantActorList)
-				{
-					Params.ConnectionManager.NotifyAddDormantDestructionInfo(Actor);
+	if (PreviousCell)
+	{
+		if (UReplicationGraphNode_DormancyNode* DormancyNode = PreviousCell->GetDormancyNode())
+		{
+			DormancyNode->ConditionalGatherDormantDynamicActors(PrevDormantActorList, Params, &DormantActorList);
+		}
+	}
 
-					if (FConnectionReplicationActorInfo* ActorInfo = Params.ConnectionManager.ActorInfoMap.Find(Actor))
-					{
-						ActorInfo->bDormantOnConnection = false;
+	// any previous dormant actors not in the current node dormant list
+	if (PrevDormantActorList.IsValid() == false)
+	{
+		return;
+	}
+
+	for (FActorRepListType& Actor : PrevDormantActorList)
+	{
+		Params.ConnectionManager.NotifyAddDormantDestructionInfo(Actor);
+
+		if (FConnectionReplicationActorInfo* ActorInfo = Params.ConnectionManager.ActorInfoMap.Find(Actor))
+		{
+			ActorInfo->bDormantOnConnection = false;
 
 						// add back to connection specific dormancy nodes
-						const FActorCellInfo CellInfo = GetCellInfoForActor(Actor, Actor->GetActorLocation(), ActorInfo->GetCullDistance());
+			const FActorCellInfo CellInfo = GetCellInfoForActor(Actor, Actor->GetActorLocation(), ActorInfo->GetCullDistance());
 
-						GetGridNodesForActor(Actor, CellInfo, GatheredNodes);
+			GetGridNodesForActor(Actor, CellInfo, GatheredNodes);
 
-						for (UReplicationGraphNode_GridCell* Node : GatheredNodes)
-						{
-							if (UReplicationGraphNode_DormancyNode* DormancyNode = Node->GetDormancyNode())
-							{
-								DormancyNode->GetConnectionNode(Params)->NotifyActorDormancyFlush(Actor);
-							}
-						}
+			for (UReplicationGraphNode_GridCell* Node : GatheredNodes)
+			{
+				if (UReplicationGraphNode_DormancyNode* DormancyNode = Node->GetDormancyNode())
+				{
+					// Only notify the connection node if this client was previously inside the cell.
+					if (UReplicationGraphNode_ConnectionDormanyNode* ConnectionDormancyNode = DormancyNode->GetExistingConnectionNode(Params))
+					{
+						ConnectionDormancyNode->NotifyActorDormancyFlush(Actor);
 					}
 				}
 			}
diff --git a/Engine/Plugins/Runtime/ReplicationGraph/Source/Public/ReplicationGraph.h b/Engine/Plugins/Runtime/ReplicationGraph/Source/Public/ReplicationGraph.h
--- a/Engine/Plugins/Runtime/ReplicationGraph/Source/Public/ReplicationGraph.h
+++ b/Engine/Plugins/Runtime/ReplicationGraph/Source/Public/ReplicationGraph.h
@@ -420,6 +420,8 @@ public:
 
 	void ConditionalGatherDormantDynamicActors(FActorRepListRefView& RepList, const FConnectionGatherActorListParameters& Params, FActorRepListRefView* RemovedList);
 
+	UReplicationGraphNode_ConnectionDormanyNode* GetExistingConnectionNode(const FConnectionGatherActorListParameters& Params);
+
 	UReplicationGraphNode_ConnectionDormanyNode* GetConnectionNode(const FConnectionGatherActorListParameters& Params);
 
 private:
@@ -516,6 +518,9 @@ public:
 	TArray<FString> DebugActorNames;
 #endif
 
+	// When enabled the RepGraph tells clients to destroy dormant dynamic actors when they go out of relevancy.
+	bool bDestroyDormantDynamicActors = true;
+
 protected:
 
 	/** For adding new actor to the graph */
@@ -582,6 +587,8 @@ private:
 	
 	void PutStaticActorIntoCell(const FNewReplicatedActorInfo& ActorInfo, FGlobalActorReplicationInfo& ActorRepInfo, bool bDormancyDriven);
 
+	void DoDormantDynamicActorDestruction(const FConnectionGatherActorListParameters& Params, UReplicationGraphNode_GridCell* CurrentCell, UReplicationGraphNode_GridCell* PreviousCell);
+
 	UReplicationGraphNode_GridCell* GetCellNode(UReplicationGraphNode_GridCell*& NodePtr)
 	{
 		if (NodePtr == nullptr)
@@ -617,7 +624,13 @@ private:
 			GridX.SetNum(Y+1);
 		}
 		return GridX[Y];
-	}			
+	}		
+
+	UReplicationGraphNode_GridCell*& GetCell(int32 X, int32 Y)
+	{
+		TArray<UReplicationGraphNode_GridCell*>& GridX = GetGridX(X);
+		return GetCell(GridX, Y);
+	}
 
 	bool bNeedsRebuild = false;
 
