Replication graph changes to allow for subconnection calculations as well, matching functionality of the netdriver system.

diff --git a/Engine/Plugins/Runtime/ReplicationGraph/Source/Private/BasicReplicationGraph.cpp b/Engine/Plugins/Runtime/ReplicationGraph/Source/Private/BasicReplicationGraph.cpp
--- a/Engine/Plugins/Runtime/ReplicationGraph/Source/Private/BasicReplicationGraph.cpp
+++ b/Engine/Plugins/Runtime/ReplicationGraph/Source/Private/BasicReplicationGraph.cpp
@@ -7,6 +7,7 @@
 #include "CoreGlobals.h"
 #include "UObject/UObjectIterator.h"
 #include "Engine/NetConnection.h"
+#include "Engine/ChildConnection.h"
 
 UBasicReplicationGraph::UBasicReplicationGraph()
 {
@@ -183,4 +184,15 @@ int32 UBasicReplicationGraph::ServerReplicateActors(float DeltaSeconds)
 
 
 	return Super::ServerReplicateActors(DeltaSeconds);
-}
\ No newline at end of file
+}
+
+bool FConnectionAlwaysRelevantNodePair::operator==(UNetConnection* InConnection) const
+{
+	// Any children should be looking at their parent connections instead.
+	if (InConnection->GetUChildConnection() != nullptr)
+	{
+		InConnection = ((UChildConnection*)InConnection)->Parent;
+	}
+
+	return InConnection == NetConnection;
+}
diff --git a/Engine/Plugins/Runtime/ReplicationGraph/Source/Private/ReplicationGraph.cpp b/Engine/Plugins/Runtime/ReplicationGraph/Source/Private/ReplicationGraph.cpp
--- a/Engine/Plugins/Runtime/ReplicationGraph/Source/Private/ReplicationGraph.cpp
+++ b/Engine/Plugins/Runtime/ReplicationGraph/Source/Private/ReplicationGraph.cpp
@@ -327,6 +327,15 @@ UNetReplicationGraphConnection* UReplicationGraph::FindOrAddConnectionManager(UN
 {
 	FScopeLogTime SLT( TEXT( "UReplicationGraph::FindOrAddConnectionManager(" ), nullptr, FScopeLogTime::ScopeLog_Milliseconds );
 
+	// Children do not have a connection manager, this is handled by their parent.
+	// We do not want to create connection managers for children, so redirect them.
+	if (NetConnection->GetUChildConnection() != nullptr)
+	{
+		NetConnection = ((UChildConnection*)NetConnection)->Parent;
+		UE_LOG(LogReplicationGraph, Warning, TEXT("UReplicationGraph::FindOrAddConnectionManager was called with a child connection, redirecting to parent"));
+		check(NetConnection != nullptr);
+	}
+
 	// Could use an acceleration map if necessary
 	RG_QUICK_SCOPE_CYCLE_COUNTER(UReplicationGraph_FindConnectionManager)
 	for (UNetReplicationGraphConnection* ConnManager : Connections)
@@ -355,6 +364,14 @@ void UReplicationGraph::AddClientConnection(UNetConnection* NetConnection)
 {
 	//FScopeLogTime SLT( TEXT( "UReplicationGraph::AddClientConnection" ), nullptr, FScopeLogTime::ScopeLog_Milliseconds );
 
+	// Children do not have a connection manager, do not proceed with this function in this case.
+	// Default behavior never calls this function with child connections anyways, so this is really only here for protection.
+	if (NetConnection->GetUChildConnection() != nullptr)
+	{
+		UE_LOG(LogReplicationGraph, Warning, TEXT("UReplicationGraph::AddClientConnection was called with a child connection, dropping."));
+		return;
+	}
+
 	// We may have already created a manager for this connection in the pending list
 	for (int32 i=PendingConnections.Num()-1; i >= 0; --i)
 	{
@@ -400,6 +417,14 @@ void UReplicationGraph::RemoveClientConnection(UNetConnection* NetConnection)
 	int32 ConnectionId = 0;
 	bool bFound = false;
 
+	// Children do not have a connection manager, do not attempt to remove it here.
+	// Default behavior never calls this function with child connections anyways, so this is really only here for protection.
+	if (NetConnection->GetUChildConnection() != nullptr)
+	{
+		UE_LOG(LogReplicationGraph, Warning, TEXT("UReplicationGraph::RemoveClientConnection was called with a child connection!"));
+		return;
+	}
+
 	// Remove the RepGraphConnection associated with this NetConnection. Also update ConnectionIds to stay compact.
 	auto UpdateList = [&](TArray<UNetReplicationGraphConnection*> List)
 	{
@@ -641,6 +666,15 @@ void UReplicationGraph::NotifyActorFullyDormantForConnection(AActor* Actor, UNet
 {
 	RG_QUICK_SCOPE_CYCLE_COUNTER(UReplicationGraph_NotifyActorFullyDormantForConnection);
 
+	// Children do not have a connection manager, so redirect as necessary.
+	// This is unlikely to be reached as child connections don't open their own channels
+	if (Connection->GetUChildConnection() != nullptr)
+	{
+		Connection = ((UChildConnection*)Connection)->Parent;
+		UE_LOG(LogReplicationGraph, Warning, TEXT("UReplicationGraph::NotifyActorFullyDormantForConnection was called for a child connection %s on actor %s"), *Connection->GetName(), *Actor->GetName());
+		check(Connection != nullptr);
+	}
+
 	// This is kind of bad but unavoidable. Possibly could use acceleration map (actor -> connections) but that would be a pain to maintain.
 	for (UNetReplicationGraphConnection* ConnectionManager: Connections)
 	{
@@ -774,38 +808,58 @@ int32 UReplicationGraph::ServerReplicateActors(float DeltaSeconds)
 	
 	FGatheredReplicationActorLists GatheredReplicationListsForConnection;
 
+	// Total number of children processed, added to all the connections later for stat tracking purposes.
+	int32 NumChildrenConnectionsProcessed = 0;
+
 	for (UNetReplicationGraphConnection* ConnectionManager: Connections)
 	{
+		// Prepare for Replication also handles children as well.
 		if (ConnectionManager->PrepareForReplication() == false)
 		{
 			// Connection is not ready to replicate
 			continue;
 		}
 
+		FNetViewerArray ConnectionViewers;
 		UNetConnection* const NetConnection = ConnectionManager->NetConnection;
 		APlayerController* const PC = NetConnection->PlayerController;
 		FPerConnectionActorInfoMap& ConnectionActorInfoMap = ConnectionManager->ActorInfoMap;
 		
 		repCheckf(NetConnection->GetReplicationConnectionDriver() == ConnectionManager, TEXT("NetConnection %s mismatch rep driver. %s vs %s"), *GetNameSafe(NetConnection), *GetNameSafe(NetConnection->GetReplicationConnectionDriver()), *GetNameSafe(ConnectionManager));
 		
+		new(ConnectionViewers) FNetViewer(NetConnection, 0.f);
+		
 		// Send ClientAdjustments (movement RPCs) do this first and never let bandwidth saturation suppress these.
 		if (PC)
 		{
 			PC->SendClientAdjustment();
 		}
 
-		FBitWriter& ConnectionSendBuffer = NetConnection->SendBuffer;
+		// Do the above but on all splitscreen connections as well.
+		for (int32 ChildIdx = 0; ChildIdx < NetConnection->Children.Num(); ++ChildIdx)
+		{
+			UNetConnection* ChildConnection = NetConnection->Children[ChildIdx];
+			if (ChildConnection && ChildConnection->PlayerController && ChildConnection->ViewTarget)
+			{
+				ChildConnection->PlayerController->SendClientAdjustment();
+				new(ConnectionViewers) FNetViewer(ChildConnection, 0.f);
+			}
+		}
+
+		NumChildrenConnectionsProcessed += NetConnection->Children.Num();
+
+		FBitWriter& ConnectionSendBuffer = NetConnection->SendBuffer; // unused
 		ConnectionManager->QueuedBitsForActorDiscovery = 0;
 
 		// --------------------------------------------------------------------------------------------------------------
 		// GATHER list of ReplicationLists for this connection
 		// --------------------------------------------------------------------------------------------------------------
 		
-		FNetViewer Viewer(NetConnection, 0.f);
-		const FVector ConnectionViewLocation = Viewer.ViewLocation;
 		GatheredReplicationListsForConnection.Reset();
 
-		const FConnectionGatherActorListParameters Parameters(Viewer, *ConnectionManager, NetConnection->ClientVisibleLevelNames, FrameNum, GatheredReplicationListsForConnection);
+		TSet<FName> AllVisibleLevelNames;
+		ConnectionManager->GetClientVisibleLevelNames(AllVisibleLevelNames);
+		const FConnectionGatherActorListParameters Parameters(ConnectionViewers, *ConnectionManager, AllVisibleLevelNames, FrameNum, GatheredReplicationListsForConnection);
 
 		{
 			QUICK_SCOPE_CYCLE_COUNTER(NET_ReplicateActors_GatherForConnection);
@@ -820,7 +874,30 @@ int32 UReplicationGraph::ServerReplicateActors(float DeltaSeconds)
 				Node->GatherActorListsForConnection(Parameters);
 			}
 
+			// Update all the LastGatherLocations for this connection.
+			for (const FNetViewer& CurViewer : ConnectionViewers)
+			{
+				FLastLocationGatherInfo* LastInfoForViewer = ConnectionManager->LastGatherLocations.FindByKey<UNetConnection*>(CurViewer.Connection);
+				if (LastInfoForViewer != nullptr)
+				{
+					LastInfoForViewer->LastLocation = CurViewer.ViewLocation;
+				}
+				else
+				{
+					// We need to add this viewer to the last gather locations
+					ConnectionManager->LastGatherLocations.Add(FLastLocationGatherInfo(CurViewer.Connection, CurViewer.ViewLocation));
+				}
+			}
+
+			// Clean up any dead entries in the last gather array
+			ConnectionManager->LastGatherLocations.RemoveAll([&](FLastLocationGatherInfo& CurGatherInfo) {
+				return CurGatherInfo.Connection == nullptr;
+			});
+
+			// Do this so we don't break anyone.
+			PRAGMA_DISABLE_DEPRECATION_WARNINGS
 			Parameters.ConnectionManager.LastGatherLocation = Parameters.Viewer.ViewLocation;
+			PRAGMA_ENABLE_DEPRECATION_WARNINGS
 
 			if (GatheredReplicationListsForConnection.NumLists() == 0)
 			{
@@ -836,8 +913,8 @@ int32 UReplicationGraph::ServerReplicateActors(float DeltaSeconds)
 		{
 			QUICK_SCOPE_CYCLE_COUNTER(NET_ReplicateActors_ProcessGatheredLists);
 
-			ReplicateActorListsForConnection_Default(ConnectionManager, GatheredReplicationListsForConnection, Viewer);
-			ReplicateActorListsForConnection_FastShared(ConnectionManager, GatheredReplicationListsForConnection, Viewer);
+			ReplicateActorListsForConnections_Default(ConnectionManager, GatheredReplicationListsForConnection, ConnectionViewers);
+			ReplicateActorListsForConnections_FastShared(ConnectionManager, GatheredReplicationListsForConnection, ConnectionViewers);
 		}
 
 		{
@@ -904,7 +981,7 @@ int32 UReplicationGraph::ServerReplicateActors(float DeltaSeconds)
 			// ------------------------------------------
 			{
 				QUICK_SCOPE_CYCLE_COUNTER(NET_ReplicateActors_ReplicateDestructionInfos);
-				ConnectionManager->ReplicateDestructionInfos(ConnectionViewLocation, DestructInfoMaxDistanceSquared);
+				ConnectionManager->ReplicateDestructionInfos(ConnectionViewers, DestructInfoMaxDistanceSquared);
 			}
 
 			// ------------------------------------------
@@ -930,7 +1007,7 @@ int32 UReplicationGraph::ServerReplicateActors(float DeltaSeconds)
 		}
 	}
 	
-	SET_DWORD_STAT(STAT_NumProcessedConnections, Connections.Num());
+	SET_DWORD_STAT(STAT_NumProcessedConnections, Connections.Num() + NumChildrenConnectionsProcessed);
 
 	if (CVar_RepGraph_PrintTrackClassReplication)
 	{
@@ -943,13 +1020,12 @@ int32 UReplicationGraph::ServerReplicateActors(float DeltaSeconds)
 	return 0;
 }
 
-void UReplicationGraph::ReplicateActorListsForConnection_Default(UNetReplicationGraphConnection* ConnectionManager, FGatheredReplicationActorLists& GatheredReplicationListsForConnection, FNetViewer& Viewer)
+void UReplicationGraph::ReplicateActorListsForConnections_Default(UNetReplicationGraphConnection* ConnectionManager, FGatheredReplicationActorLists& GatheredReplicationListsForConnection, FNetViewerArray& Viewers)
 {
-
 #if !(UE_BUILD_SHIPPING || UE_BUILD_TEST)
 	const bool bEnableFullActorPrioritizationDetails = DO_REPGRAPH_DETAILS(bEnableFullActorPrioritizationDetailsAllConnections || ConnectionManager->bEnableFullActorPrioritizationDetails);
 	const bool bDoDistanceCull = (CVar_RepGraph_SkipDistanceCull == 0);
-	const bool bDoCulledOnConnectionCount = (CVar_RepGraph_PrintCulledOnConnectionClasses == 1);	
+	const bool bDoCulledOnConnectionCount = (CVar_RepGraph_PrintCulledOnConnectionClasses == 1);
 	bTrackClassReplication = (CVar_RepGraph_TrackClassReplication > 0 || CVar_RepGraph_PrintTrackClassReplication > 0);
 	if (!bTrackClassReplication)
 	{
@@ -975,8 +1051,6 @@ void UReplicationGraph::ReplicateActorListsForConnection_Default(UNetReplication
 	FPerConnectionActorInfoMap& ConnectionActorInfoMap = ConnectionManager->ActorInfoMap;
 	const uint32 FrameNum = ReplicationGraphFrame;
 
-	const FVector ConnectionViewLocation = Viewer.ViewLocation;
-
 	// --------------------------------------------------------------------------------------------------------------
 	// PRIORITIZE Gathered Actors For Connection
 	// --------------------------------------------------------------------------------------------------------------
@@ -990,9 +1064,10 @@ void UReplicationGraph::ReplicateActorListsForConnection_Default(UNetReplication
 		TArray<FPrioritizedRepList::FItem>* SortingArray = &PrioritizedReplicationList.Items;
 
 		NumGatheredListsOnConnection += GatheredReplicationListsForConnection.NumLists();
-				
+
 		const float MaxDistanceScaling = PrioritizationConstants.MaxDistanceScaling;
 		const uint32 MaxFramesSinceLastRep = PrioritizationConstants.MaxFramesSinceLastRep;
+		const int32 TotalNumOfConnections = 1 + NetConnection->Children.Num();
 
 		for (FActorRepListRawView& List : GatheredReplicationListsForConnection.GetLists(EActorRepListTypeFlags::Default))
 		{
@@ -1040,7 +1115,7 @@ void UReplicationGraph::ReplicateActorListsForConnection_Default(UNetReplication
 				}
 
 				RG_QUICK_SCOPE_CYCLE_COUNTER(Prioritize_InnerLoop_PostReady);
-							
+
 				// Output record for full debugging. This is not used in the actual sorting/prioritization of the list, just for logging/debugging purposes
 				FPrioritizedActorFullDebugDetails* DebugDetails = nullptr;
 				if (DO_REPGRAPH_DETAILS(UNLIKELY(bEnableFullActorPrioritizationDetails)))
@@ -1055,11 +1130,28 @@ void UReplicationGraph::ReplicateActorListsForConnection_Default(UNetReplication
 				// -------------------
 				if (GlobalData.Settings.DistancePriorityScale > 0.f)
 				{
-					const float DistSq = (GlobalData.WorldLocation - ConnectionViewLocation).SizeSquared();
+					float SmallestDistanceSq = TNumericLimits<float>::Max();
+					int32 ConnectionsThatSkipActor = 0;
+					
+					for (const FNetViewer& CurViewer : Viewers)
+					{
+						const float DistSq = (GlobalData.WorldLocation - CurViewer.ViewLocation).SizeSquared();
+						SmallestDistanceSq = FMath::Min<float>(DistSq, SmallestDistanceSq);
+
+						// Figure out if we should be skipping this actor
+						if (bDoDistanceCull && ConnectionData.GetCullDistanceSquared() > 0.f && DistSq > ConnectionData.GetCullDistanceSquared())
+						{
+							++ConnectionsThatSkipActor;
+							continue;
+						}
+					}
 
-					if (bDoDistanceCull && ConnectionData.GetCullDistanceSquared() > 0.f && DistSq > ConnectionData.GetCullDistanceSquared())
+					// If no one is near this actor, skip it.
+					if (ConnectionsThatSkipActor >= TotalNumOfConnections)
 					{
-						DO_REPGRAPH_DETAILS(PrioritizedReplicationList.GetNextSkippedDebugDetails(Actor)->DistanceCulled = FMath::Sqrt(DistSq));
+						DO_REPGRAPH_DETAILS(PrioritizedReplicationList.GetNextSkippedDebugDetails(Actor)->DistanceCulled = FMath::Sqrt(SmallestDistanceSq));
+
+						// Skipped actors should not have any 
 						if (bDoCulledOnConnectionCount)
 						{
 							DistanceClassAccumulator.Increment(Actor->GetClass());
@@ -1067,14 +1159,14 @@ void UReplicationGraph::ReplicateActorListsForConnection_Default(UNetReplication
 						continue;
 					}
 
-					const float DistanceFactor = FMath::Clamp<float>( (DistSq) / MaxDistanceScaling, 0.f, 1.f) * GlobalData.Settings.DistancePriorityScale;
-					AccumulatedPriority += DistanceFactor;
-								
+					const float DistanceFactor = FMath::Clamp<float>((SmallestDistanceSq) / MaxDistanceScaling, 0.f, 1.f) * GlobalData.Settings.DistancePriorityScale;
 					if (DO_REPGRAPH_DETAILS(UNLIKELY(DebugDetails)))
 					{
-						DebugDetails->DistanceSq = DistSq;
+						DebugDetails->DistanceSq = SmallestDistanceSq;
 						DebugDetails->DistanceFactor = DistanceFactor;
 					}
+
+					AccumulatedPriority += DistanceFactor;
 				}
 
 				RG_QUICK_SCOPE_CYCLE_COUNTER(Prioritize_InnerLoop_PostCull);
@@ -1117,16 +1209,21 @@ void UReplicationGraph::ReplicateActorListsForConnection_Default(UNetReplication
 				// -------------------
 				// Always prioritize the connection's owner and view target, since these are the most important actors for the client.
 				// -------------------
-				if (Actor == Viewer.ViewTarget || Actor == Viewer.InViewer)
+				for (const FNetViewer& CurViewer : Viewers)
 				{
-					AccumulatedPriority -= 10.0f;
+					// We need to find if this is anyone's viewer or viewtarget, not just the parent connection.
+					if (Actor == CurViewer.ViewTarget || Actor == CurViewer.InViewer)
+					{
+						AccumulatedPriority -= 10.0f;
+						break;
+					}
 				}
 
 				// -------------------
 				//	Game code priority
 				// -------------------
-				
-				if ( GlobalData.ForceNetUpdateFrame > ConnectionData.LastRepFrameNum )
+
+				if (GlobalData.ForceNetUpdateFrame > ConnectionData.LastRepFrameNum)
 				{
 					// Note that in legacy ForceNetUpdate did not actually bump priority. This gives us a hard coded bump if we haven't replicated since the last ForceNetUpdate frame.
 					AccumulatedPriority -= 1.f;
@@ -1136,8 +1233,8 @@ void UReplicationGraph::ReplicateActorListsForConnection_Default(UNetReplication
 						DebugDetails->GameCodeScaling = -1.f;
 					}
 				}
-							
-				SortingArray->Emplace(FPrioritizedRepList::FItem(AccumulatedPriority, Actor, &GlobalData, &ConnectionData ));
+
+				SortingArray->Emplace(FPrioritizedRepList::FItem(AccumulatedPriority, Actor, &GlobalData, &ConnectionData));
 			}
 		}
 
@@ -1155,7 +1252,7 @@ void UReplicationGraph::ReplicateActorListsForConnection_Default(UNetReplication
 	{
 		QUICK_SCOPE_CYCLE_COUNTER(NET_ReplicateActors_ReplicateActorsForConnection);
 
-		for (int32 ActorIdx=0; ActorIdx < PrioritizedReplicationList.Items.Num(); ++ActorIdx)
+		for (int32 ActorIdx = 0; ActorIdx < PrioritizedReplicationList.Items.Num(); ++ActorIdx)
 		{
 			const FPrioritizedRepList::FItem& RepItem = PrioritizedReplicationList.Items[ActorIdx];
 
@@ -1165,7 +1262,7 @@ void UReplicationGraph::ReplicateActorListsForConnection_Default(UNetReplication
 			// Always skip if we've already replicated this frame. This happens if an actor is in more than one replication list
 			if (ActorInfo.LastRepFrameNum == FrameNum)
 			{
-				INC_DWORD_STAT_BY( STAT_NetRepActorListDupes, 1 );
+				INC_DWORD_STAT_BY(STAT_NetRepActorListDupes, 1);
 				continue;
 			}
 
@@ -1176,12 +1273,12 @@ void UReplicationGraph::ReplicateActorListsForConnection_Default(UNetReplication
 			// --------------------------------------------------
 			//	Update Packet Budget Tracking
 			// --------------------------------------------------
-					
+
 			if (IsConnectionReady(NetConnection) == false)
 			{
 				// We've exceeded the budget for this category of replication list.
 				RG_QUICK_SCOPE_CYCLE_COUNTER(NET_ReplicateActors_PartialStarvedActorList);
-				HandleStarvedActorList(PrioritizedReplicationList, ActorIdx+1, ConnectionActorInfoMap, FrameNum);
+				HandleStarvedActorList(PrioritizedReplicationList, ActorIdx + 1, ConnectionActorInfoMap, FrameNum);
 				GNumSaturatedConnections++;
 				break;
 			}
@@ -1189,7 +1286,7 @@ void UReplicationGraph::ReplicateActorListsForConnection_Default(UNetReplication
 	}
 
 	// Broadcast the list we just handled. This is intended to be for debugging/logging features.
-	ConnectionManager->OnPostReplicatePrioritizeLists.Broadcast( ConnectionManager, &PrioritizedReplicationList);
+	ConnectionManager->OnPostReplicatePrioritizeLists.Broadcast(ConnectionManager, &PrioritizedReplicationList);
 
 	if (bDoCulledOnConnectionCount)
 	{
@@ -1203,6 +1300,22 @@ void UReplicationGraph::ReplicateActorListsForConnection_Default(UNetReplication
 	}
 }
 
+PRAGMA_DISABLE_DEPRECATION_WARNINGS
+void UReplicationGraph::ReplicateActorListsForConnection_Default(UNetReplicationGraphConnection* ConnectionManager, FGatheredReplicationActorLists& GatheredReplicationListsForConnection, FNetViewer& Viewer)
+{
+	FNetViewerArray ViewersToConsider;
+	ViewersToConsider.Add(Viewer);
+
+	// Create viewers for all the children related to this connection
+	for (int32 ChildIdx = 0; ChildIdx < Viewer.Connection->Children.Num(); ++ChildIdx)
+	{
+		new(ViewersToConsider)FNetViewer(Viewer.Connection->Children[ChildIdx], 0);
+	}
+
+	ReplicateActorListsForConnections_Default(ConnectionManager, GatheredReplicationListsForConnection, ViewersToConsider);
+}
+PRAGMA_ENABLE_DEPRECATION_WARNINGS
+
 struct FScopedQueuedBits
 {
 	FScopedQueuedBits(int32& InQueuedBits, int32& InTotalBits) : QueuedBits(InQueuedBits), TotalBits(InTotalBits) { }
@@ -1252,10 +1365,9 @@ struct FScopedFastPathTracker
 #endif
 };
 
-
-void UReplicationGraph::ReplicateActorListsForConnection_FastShared(UNetReplicationGraphConnection* ConnectionManager, FGatheredReplicationActorLists& GatheredReplicationListsForConnection, FNetViewer& Viewer)
+void UReplicationGraph::ReplicateActorListsForConnections_FastShared(UNetReplicationGraphConnection* ConnectionManager, FGatheredReplicationActorLists& GatheredReplicationListsForConnection, FNetViewerArray& Viewers)
 {
-	if(CVar_RepGraph_EnableFastSharedPath == false)
+	if (CVar_RepGraph_EnableFastSharedPath == false)
 	{
 		return;
 	}
@@ -1268,8 +1380,6 @@ void UReplicationGraph::ReplicateActorListsForConnection_FastShared(UNetReplicat
 	FPerConnectionActorInfoMap& ConnectionActorInfoMap = ConnectionManager->ActorInfoMap;
 	UNetConnection* const NetConnection = ConnectionManager->NetConnection;
 	const uint32 FrameNum = ReplicationGraphFrame;
-	const FVector ConnectionViewLocation = Viewer.ViewLocation;
-	const FVector ConnectionViewDir = Viewer.ViewDir;
 	const float FastSharedDistanceRequirementPct = FastSharedPathConstants.DistanceRequirementPct;
 	const int64 MaxBits = FastSharedPathConstants.MaxBitsPerFrame;
 	const int32 StartIdx = FrameNum * FastSharedPathConstants.ListSkipPerFrame;
@@ -1281,10 +1391,10 @@ void UReplicationGraph::ReplicateActorListsForConnection_FastShared(UNetReplicat
 	FScopedQueuedBits ScopedQueuedBits(NetConnection->QueuedBits, TotalBitsWritten);
 
 	TArray< FActorRepListRawView>& GatheredLists = GatheredReplicationListsForConnection.GetLists(EActorRepListTypeFlags::FastShared);
-	for (int32 ListIdx=0; ListIdx < GatheredLists.Num(); ++ListIdx)
+	for (int32 ListIdx = 0; ListIdx < GatheredLists.Num(); ++ListIdx)
 	{
 		FActorRepListRawView& List = GatheredLists[(ListIdx + FrameNum) % GatheredLists.Num()];
-		for (int32 i=0; i < List.Num(); ++i)
+		for (int32 i = 0; i < List.Num(); ++i)
 		{
 			// Round robin through the list over multiple frames. We want to avoid sorting this list based on 'time since last rep'. This is a good balance
 			AActor* Actor = List[(i + StartIdx) % List.Num()];
@@ -1324,16 +1434,32 @@ void UReplicationGraph::ReplicateActorListsForConnection_FastShared(UNetReplicat
 				continue;
 			}
 
-			// Simple dot product rejection: only fast rep actors in front of this connection
-			const FVector DirToActor = GlobalActorInfo.WorldLocation - ConnectionViewLocation;
-			if (FVector::DotProduct(DirToActor, ConnectionViewDir) < 0.f)
+			// Determine if this actor has any view relevancy to any connection this client has
+			bool bNoViewRelevency = true;
+			for (const FNetViewer& CurView : Viewers)
 			{
-				continue;
+				const FVector& ConnectionViewLocation = CurView.ViewLocation;
+				const FVector& ConnectionViewDir = CurView.ViewDir;
+
+				// Simple dot product rejection: only fast rep actors in front of this connection
+				const FVector DirToActor = GlobalActorInfo.WorldLocation - ConnectionViewLocation;
+				if (!(FVector::DotProduct(DirToActor, ConnectionViewDir) < 0.f))
+				{
+					bNoViewRelevency = false;
+					break;
+				}
+
+				// Simple distance cull
+				const float DistSq = DirToActor.SizeSquared();
+				if (!(DistSq > (ConnectionData.GetCullDistanceSquared() * FastSharedDistanceRequirementPct)))
+				{
+					bNoViewRelevency = false;
+					break;
+				}
 			}
 
-			// Simple distance cull
-			const float DistSq = DirToActor.SizeSquared();
-			if (DistSq > (ConnectionData.GetCullDistanceSquared() * FastSharedDistanceRequirementPct))
+			// Skip out if they have none.
+			if (bNoViewRelevency)
 			{
 				continue;
 			}
@@ -1357,6 +1483,22 @@ void UReplicationGraph::ReplicateActorListsForConnection_FastShared(UNetReplicat
 	}
 }
 
+PRAGMA_DISABLE_DEPRECATION_WARNINGS
+void UReplicationGraph::ReplicateActorListsForConnection_FastShared(UNetReplicationGraphConnection* ConnectionManager, FGatheredReplicationActorLists& GatheredReplicationListsForConnection, FNetViewer& Viewer)
+{
+	FNetViewerArray ViewersToConsider;
+	ViewersToConsider.Add(Viewer);
+
+	// Create viewers for all the children related to this connection
+	for (int32 ChildIdx = 0; ChildIdx < Viewer.Connection->Children.Num(); ++ChildIdx)
+	{
+		new(ViewersToConsider)FNetViewer(Viewer.Connection->Children[ChildIdx], 0);
+	}
+
+	ReplicateActorListsForConnections_FastShared(ConnectionManager, GatheredReplicationListsForConnection, ViewersToConsider);
+}
+PRAGMA_ENABLE_DEPRECATION_WARNINGS
+
 
 REPGRAPH_DEVCVAR_SHIPCONST(int32, "Net.RepGraph.FastShared.ForceFull", CVar_RepGraph_FastShared_ForceFull, 0, "Redirects calls to ReplicateSingleActor_FastShared to ReplicateSingleActor");
 
@@ -1807,16 +1949,33 @@ bool UReplicationGraph::ProcessRemoteFunction(class AActor* Actor, UFunction* Fu
 					bool ShouldOpenChannel = true;
 					if (CullDistanceSquared > 0.f)
 					{
-						FNetViewer Viewer(NetConnection, 0.f);
+						ShouldOpenChannel = false;
 						if (ActorLocation.IsSet() == false)
 						{
 							ActorLocation = Actor->GetActorLocation();
 						}
 
-						const float DistSq = (ActorLocation.GetValue() - Viewer.ViewLocation).SizeSquared();
-						if (DistSq > CullDistanceSquared)
+						FNetViewerArray ViewsToConsider;
+						new(ViewsToConsider)FNetViewer(NetConnection, 0.f);
+
+						for (int32 ChildIdx = 0; ChildIdx < NetConnection->Children.Num(); ++ChildIdx)
 						{
-							ShouldOpenChannel = false;
+							if (NetConnection->Children[ChildIdx]->ViewTarget != nullptr)
+							{
+								new(ViewsToConsider)FNetViewer(NetConnection->Children[ChildIdx], 0.f);
+							}
+						}
+
+						// Loop through and see if we should keep this channel open, as when we do distance, we will
+						// default to the channel being closed.
+						for (const FNetViewer& Viewer : ViewsToConsider)
+						{
+							const float DistSq = (ActorLocation.GetValue() - Viewer.ViewLocation).SizeSquared();
+							if (DistSq <= CullDistanceSquared)
+							{
+								ShouldOpenChannel = true;
+								break;
+							}
 						}
 					}
 
@@ -2128,6 +2287,14 @@ void UNetReplicationGraphConnection::RemoveConnectionGraphNode(UReplicationGraph
 bool UNetReplicationGraphConnection::PrepareForReplication()
 {
 	NetConnection->ViewTarget = NetConnection->PlayerController ? NetConnection->PlayerController->GetViewTarget() : NetConnection->OwningActor;
+
+	// Set any children viewtargets
+	for (int32 i = 0; i < NetConnection->Children.Num(); ++i)
+	{
+		UNetConnection* CurChild = NetConnection->Children[i];
+		CurChild->ViewTarget = CurChild->PlayerController ? CurChild->PlayerController->GetViewTarget() : CurChild->OwningActor;
+	}
+
 	return (NetConnection->State != USOCK_Closed) && (NetConnection->ViewTarget != nullptr);
 }
 
@@ -2216,6 +2383,24 @@ void UNetReplicationGraphConnection::NotifyResetDestructionInfo()
 	PendingDestructInfoList.Reset();
 }
 
+void UNetReplicationGraphConnection::GetClientVisibleLevelNames(TSet<FName>& OutLevelNames) const
+{
+	if (NetConnection == nullptr)
+	{
+		return;
+	}
+
+	OutLevelNames.Append(NetConnection->ClientVisibleLevelNames);
+	for (const UNetConnection* Child : NetConnection->Children)
+	{
+		if (Child != nullptr)
+		{
+			// For sets, we don't have to worry about uniqueness due to the nature of the data structure.
+			OutLevelNames.Append(Child->ClientVisibleLevelNames);
+		}
+	}
+}
+
 void UNetReplicationGraphConnection::NotifyClientVisibleLevelNamesAdd(FName LevelName, UWorld* StreamingWorld) 
 {
 	RG_QUICK_SCOPE_CYCLE_COUNTER(UNetReplicationGraphConnection_NotifyClientVisibleLevelNamesAdd);
@@ -2243,26 +2428,37 @@ void UNetReplicationGraphConnection::NotifyClientVisibleLevelNamesAdd(FName Leve
 
 }
 
-int64 UNetReplicationGraphConnection::ReplicateDestructionInfos(const FVector& ConnectionViewLocation, const float DestructInfoMaxDistanceSquared)
+int64 UNetReplicationGraphConnection::ReplicateDestructionInfos(const FNetViewerArray& Viewers, const float DestructInfoMaxDistanceSquared)
 {
 	CSV_SCOPED_TIMING_STAT_EXCLUSIVE(ReplicateDestructionInfos);
 
-	const float X = ConnectionViewLocation.X;
-	const float Y = ConnectionViewLocation.Y;
-
 	int64 NumBits = 0;
 	for (int32 idx=PendingDestructInfoList.Num()-1; idx >=0; --idx)
 	{
 		const FCachedDestructInfo& Info = PendingDestructInfoList[idx];
 		FActorDestructionInfo* DestructInfo = Info.DestructionInfo;
-		float DistSquared = FMath::Square(Info.CachedPosition.X-X) + FMath::Square(Info.CachedPosition.Y-Y);
+		bool bStillInRelevency = false;
 
-		if (DistSquared < DestructInfoMaxDistanceSquared)
+		// Find if anyone is close to this object.
+		for (const FNetViewer& CurViewer : Viewers)
 		{
-			UActorChannel* Channel = (UActorChannel*)NetConnection->CreateChannelByName( NAME_Actor, EChannelCreateFlags::OpenedLocally );
-			if ( Channel )
+			float DistSquared = FMath::Square(Info.CachedPosition.X - CurViewer.ViewLocation.X) + FMath::Square(Info.CachedPosition.Y - CurViewer.ViewLocation.Y);
+
+			// Someone is nearby this object, do not remove it.
+			if (!(DistSquared < DestructInfoMaxDistanceSquared))
 			{
-				NumBits += Channel->SetChannelActorForDestroy( DestructInfo );
+				bStillInRelevency = true;
+				break;
+			}
+		}
+
+		// Essentially, if no one can see this object, mark it for destruction.
+		if (!bStillInRelevency)
+		{
+			UActorChannel* Channel = (UActorChannel*)NetConnection->CreateChannelByName(NAME_Actor, EChannelCreateFlags::OpenedLocally);
+			if (Channel)
+			{
+				NumBits += Channel->SetChannelActorForDestroy(DestructInfo);
 			}
 
 			PendingDestructInfoList.RemoveAtSwap(idx, 1, false);
@@ -2894,9 +3090,6 @@ void UReplicationGraphNode_DynamicSpatialFrequency::GatherActorListsForConnectio
 				SortedReplicationList.Sort();
 				SortedReplicationList.SetNum(MaxNearestActors, false);
 
-				const FVector ConnectionViewLocation = Params.Viewer.ViewLocation;
-				const FVector ConnectionViewDir = Params.Viewer.ViewDir;
-
 				// Do rest of normal spatial calculations and resort
 				for (int32 idx = SortedReplicationList.Num()-1; idx >= 0; --idx)
 				{
@@ -2905,7 +3098,7 @@ void UReplicationGraphNode_DynamicSpatialFrequency::GatherActorListsForConnectio
 					FGlobalActorReplicationInfo& GlobalInfo = *Item.GlobalInfo;
 					FConnectionReplicationActorInfo& ConnectionInfo = ConnectionActorInfoMap.FindOrAdd(Actor);
 				
-					CalcFrequencyForActor(Actor, RepGraph, NetConnection, GlobalInfo, ConnectionInfo, MySettings, ConnectionViewLocation, ConnectionViewDir, FrameNum, idx);
+					CalcFrequencyForActor(Actor, RepGraph, NetConnection, GlobalInfo, ConnectionInfo, MySettings, Params.Viewers, FrameNum, idx);
 				}
 
 				SortedReplicationList.Sort();
@@ -3047,9 +3240,7 @@ FORCEINLINE uint32 CalcDynamicReplicationPeriod(const float FinalPCT, const uint
 
 static TArray<FColor> DynamicSpatialFrequencyDebugColorArray = { FColor::Red, FColor::Green, FColor::Blue, FColor::Cyan, FColor::Orange, FColor::Purple };
 
-
-
-FORCEINLINE void UReplicationGraphNode_DynamicSpatialFrequency::CalcFrequencyForActor(AActor* Actor, UReplicationGraph* RepGraph, UNetConnection* NetConnection, FGlobalActorReplicationInfo& GlobalInfo, FConnectionReplicationActorInfo& ConnectionInfo, FSettings& MySettings, const FVector& ConnectionViewLocation, const FVector& ConnectionViewDir, const uint32 FrameNum, int32 ExistingItemIndex)
+FORCEINLINE void UReplicationGraphNode_DynamicSpatialFrequency::CalcFrequencyForActor(AActor* Actor, UReplicationGraph* RepGraph, UNetConnection* NetConnection, FGlobalActorReplicationInfo& GlobalInfo, FConnectionReplicationActorInfo& ConnectionInfo, FSettings& MySettings, const FNetViewerArray& Viewers, const uint32 FrameNum, int32 ExistingItemIndex)
 {
 	// If we need to filter out the actor and he is already in the SortedReplicationList, we need to remove it (instead of just skipping/returning).
 	auto RemoveExistingItem = [&ExistingItemIndex, this]()
@@ -3081,12 +3272,27 @@ FORCEINLINE void UReplicationGraphNode_DynamicSpatialFrequency::CalcFrequencyFor
 	{
 		RemoveExistingItem();
 		return;
-	}	
+	}
+
+	float SmallestDistanceToActorSq = TNumericLimits<float>::Max();
+	const FNetViewer* LowestDistanceViewer = nullptr;
+
+	// Find the closest viewer to this item
+	for (const FNetViewer& CurViewer : Viewers)
+	{
+		float CurDistance = (GlobalInfo.WorldLocation - CurViewer.ViewLocation).SizeSquared();
+		if (CurDistance < SmallestDistanceToActorSq)
+		{
+			LowestDistanceViewer = &CurViewer;
+			SmallestDistanceToActorSq = CurDistance;
+		}
+	}
+
+	check(LowestDistanceViewer != nullptr);
+	UE_LOG(LogReplicationGraph, VeryVerbose, TEXT("UReplicationGraphNode_DynamicSpatialFrequency::CalcFrequencyForActor: Using viewer %s for spatical determination for actor %s"), ((LowestDistanceViewer->Connection != nullptr) ? *LowestDistanceViewer->Connection->GetName() : TEXT("INVALID")), *Actor->GetName());
 
 	// Skip if past cull distance
-	const FVector DirToActor = GlobalInfo.WorldLocation - ConnectionViewLocation;
-	const float DistanceToActorSq = DirToActor.SizeSquared();
-	if (!IgnoreCullDistance && ConnectionInfo.GetCullDistanceSquared() > 0.f && DistanceToActorSq > ConnectionInfo.GetCullDistanceSquared())
+	if (!IgnoreCullDistance && ConnectionInfo.GetCullDistanceSquared() > 0.f && SmallestDistanceToActorSq > ConnectionInfo.GetCullDistanceSquared())
 	{
 		RemoveExistingItem();
 		return;
@@ -3095,7 +3301,9 @@ FORCEINLINE void UReplicationGraphNode_DynamicSpatialFrequency::CalcFrequencyFor
 	// --------------------------------------------------------------------------------------------------------
 	// Find Zone
 	// --------------------------------------------------------------------------------------------------------
-	const float DistanceToActor = FMath::Sqrt(DistanceToActorSq);
+	const FVector& ConnectionViewDir = LowestDistanceViewer->ViewDir;
+	const FVector DirToActor = (GlobalInfo.WorldLocation - LowestDistanceViewer->ViewLocation);
+	const float DistanceToActor = FMath::Sqrt(SmallestDistanceToActorSq);
 	const FVector NormDirToActor = DistanceToActor > SMALL_NUMBER ? (DirToActor / DistanceToActor) : DirToActor;
 	const float DotP = FVector::DotProduct(NormDirToActor, ConnectionViewDir);
 
@@ -3169,7 +3377,7 @@ FORCEINLINE void UReplicationGraphNode_DynamicSpatialFrequency::CalcFrequencyFor
 			// We now know when this actor should replicate next. We either need to add or remove the item from the sorted list. 
 			// We also may need to do some tracking for replicate this frame vs next (for Opportunistic LoadBalance)
 			// --------------------------------------------------------------------------------------------------------
-			
+
 			if (FramesTillReplicate < 0)
 			{
 				// This actor is ready to go (or overdue). Add him to the replication list that we will sort
@@ -3210,22 +3418,49 @@ FORCEINLINE void UReplicationGraphNode_DynamicSpatialFrequency::CalcFrequencyFor
 	}
 
 	// No zone was found. This is bad
-	UE_LOG(LogReplicationGraph, Warning, TEXT("UReplicationGraphNode_DynamicSpatialFrequency::CalcFrequencyForActor: %s was not placed in any valid zone. DotP: %.2f "), *Actor->GetName(), DotP);
+	UE_LOG(LogReplicationGraph, Warning, TEXT("UReplicationGraphNode_DynamicSpatialFrequency::CalcFrequencyForActor: %s was not placed in any valid zone. Viewer: %s DotP: %.2f "), *Actor->GetName(), ((LowestDistanceViewer->Connection != nullptr) ? *LowestDistanceViewer->Connection->GetName() : TEXT("INVALID")), DotP);
 	RemoveExistingItem();
 }
 
+PRAGMA_DISABLE_DEPRECATION_WARNINGS
+FORCEINLINE void UReplicationGraphNode_DynamicSpatialFrequency::CalcFrequencyForActor(AActor* Actor, UReplicationGraph* RepGraph, UNetConnection* NetConnection, FGlobalActorReplicationInfo& GlobalInfo, FConnectionReplicationActorInfo& ConnectionInfo, FSettings& MySettings, const FVector& ConnectionViewLocation, const FVector& ConnectionViewDir, const uint32 FrameNum, int32 ExistingItemIndex)
+{
+	FNetViewerArray ViewersArray;
+	FNetViewer SoloViewer(NetConnection, 0.f);
+	SoloViewer.ViewLocation = ConnectionViewLocation;
+	SoloViewer.ViewDir = ConnectionViewDir;
+	ViewersArray.Add(SoloViewer);
+
+	// Add the child viewers
+	for (UNetConnection* Child : NetConnection->Children)
+	{
+		new(ViewersArray) FNetViewer(Child, 0.f);
+	}
+
+	CalcFrequencyForActor(Actor, RepGraph, NetConnection, GlobalInfo, ConnectionInfo, MySettings, ViewersArray, FrameNum, ExistingItemIndex);
+}
+PRAGMA_ENABLE_DEPRECATION_WARNINGS
+
 void UReplicationGraphNode_DynamicSpatialFrequency::GatherActors(const FActorRepListRefView& RepList, FGlobalActorReplicationInfoMap& GlobalMap, FPerConnectionActorInfoMap& ConnectionMap, const FConnectionGatherActorListParameters& Params, UNetConnection* NetConnection)
 {
 	UReplicationGraph* RepGraph = GraphGlobals->ReplicationGraph;
 	FSettings& MySettings = GetSettings();
-	const FVector ConnectionViewLocation = Params.Viewer.ViewLocation;
-	const FVector ConnectionViewDir = Params.Viewer.ViewDir;
 	const uint32 FrameNum = Params.ReplicationFrameNum;	
 
 	for (AActor* Actor : RepList)
 	{
+		bool bShouldSkipActor = false;
 		// Don't replicate the connection view target like this. It will be done through a connection specific node
-		if (UNLIKELY(Actor == Params.Viewer.ViewTarget))
+		for (const FNetViewer& CurViewer : Params.Viewers)
+		{
+			if (UNLIKELY(Actor == CurViewer.ViewTarget))
+			{
+				bShouldSkipActor = true;
+				break;
+			}
+		}
+
+		if (bShouldSkipActor)
 		{
 			continue;
 		}
@@ -3233,27 +3468,37 @@ void UReplicationGraphNode_DynamicSpatialFrequency::GatherActors(const FActorRep
 		FGlobalActorReplicationInfo& GlobalInfo = GlobalMap.Get(Actor);
 		FConnectionReplicationActorInfo& ConnectionInfo = ConnectionMap.FindOrAdd(Actor);
 
-		CalcFrequencyForActor(Actor, RepGraph, NetConnection, GlobalInfo, ConnectionInfo, MySettings, ConnectionViewLocation, ConnectionViewDir, FrameNum, INDEX_NONE);
+		CalcFrequencyForActor(Actor, RepGraph, NetConnection, GlobalInfo, ConnectionInfo, MySettings, Params.Viewers, FrameNum, INDEX_NONE);
 	}
 }
 
 void UReplicationGraphNode_DynamicSpatialFrequency::GatherActors_DistanceOnly(const FActorRepListRefView& RepList, FGlobalActorReplicationInfoMap& GlobalMap, FPerConnectionActorInfoMap& ConnectionMap, const FConnectionGatherActorListParameters& Params)
 {
 	FGlobalActorReplicationInfoMap* GlobalActorReplicationInfoMap = GraphGlobals->GlobalActorReplicationInfoMap;
-	const FVector ConnectionViewLocation = Params.Viewer.ViewLocation;
 	for (AActor* Actor : RepList)
 	{
+		FGlobalActorReplicationInfo& GlobalInfo = GlobalMap.Get(Actor);
+		float ShortestDistanceToActorSq = TNumericLimits<float>::Max();
+		bool bShouldSkipActor = false;
+
 		// Don't replicate the connection view target like this. It will be done through a connection specific node
-		if (UNLIKELY(Actor == Params.Viewer.ViewTarget))
+		for (const FNetViewer& CurViewer : Params.Viewers)
 		{
-			continue;
+			if (UNLIKELY(Actor == CurViewer.ViewTarget))
+			{
+				bShouldSkipActor = true;
+				break;
+			}
+
+			ShortestDistanceToActorSq = FMath::Min<float>((GlobalInfo.WorldLocation - CurViewer.ViewLocation).SizeSquared(), ShortestDistanceToActorSq);
 		}
 
-		FGlobalActorReplicationInfo& GlobalInfo = GlobalMap.Get(Actor);
-		const FVector DirToActor = GlobalInfo.WorldLocation - ConnectionViewLocation;
-		const float DistanceToActorSq = DirToActor.SizeSquared();
+		if (bShouldSkipActor)
+		{
+			continue;
+		}
 
-		SortedReplicationList.Emplace(Actor, (int32)DistanceToActorSq, &GlobalInfo);
+		SortedReplicationList.Emplace(Actor, (int32)ShortestDistanceToActorSq, &GlobalInfo);
 	}
 }
 
@@ -3496,7 +3741,17 @@ void UReplicationGraphNode_DormancyNode::RemoveDormantActor(const FNewReplicated
 
 void UReplicationGraphNode_DormancyNode::GatherActorListsForConnection(const FConnectionGatherActorListParameters& Params)
 {
-	if (Params.Viewer.ViewLocation.Z > MaxZForConnection)
+	int32 NumViewersAboveMaxZ = 0;
+	for (const FNetViewer& CurViewer : Params.Viewers)
+	{
+		if (CurViewer.ViewLocation.Z > MaxZForConnection)
+		{
+			++NumViewersAboveMaxZ;
+		}
+	}
+
+	// If we're above max on all viewers, don't gather actors.
+	if (Params.Viewers.Num() <= NumViewersAboveMaxZ)
 	{
 		return;
 	}
@@ -3565,7 +3820,7 @@ void UReplicationGraphNode_DormancyNode::OnActorDormancyFlush(FActorRepListType
 	}
 }
 
-void UReplicationGraphNode_DormancyNode::ConditionalGatherDormantDynamicActors(FActorRepListRefView& RepList, const FConnectionGatherActorListParameters& Params, FActorRepListRefView* RemovedList)
+void UReplicationGraphNode_DormancyNode::ConditionalGatherDormantDynamicActors(FActorRepListRefView& RepList, const FConnectionGatherActorListParameters& Params, FActorRepListRefView* RemovedList, bool bEnforceReplistUniqueness)
 {
 	for (FActorRepListType& Actor : ReplicationActorList)
 	{
@@ -3580,6 +3835,19 @@ void UReplicationGraphNode_DormancyNode::ConditionalGatherDormantDynamicActors(F
 						continue;
 					}
 
+					// Prevent adding actors if we already have added them, this saves on grow operations.
+					if (bEnforceReplistUniqueness)
+					{
+						if (Info->bGridSpatilization_AlreadyDormant)
+						{
+							continue;
+						}
+						else
+						{
+							Info->bGridSpatilization_AlreadyDormant = true;
+						}
+					}
+
 					RepList.PrepareForWrite();
 					RepList.ConditionalAdd(Actor);
 				}
@@ -4108,7 +4376,6 @@ static FAutoConsoleVariableRef CVarRepGraphSpatialDebugDynamic(TEXT("Net.RepGrap
 int32 CVar_RepGraph_Spatial_BiasCreep = 0.f;
 static FAutoConsoleVariableRef CVarRepGraphSpatialBiasCreep(TEXT("Net.RepGraph.Spatial.BiasCreep"), CVar_RepGraph_Spatial_BiasCreep, TEXT("Changes bias each frame by this much and force rebuld. For stress test debugging"), ECVF_Default );
 
-
 void UReplicationGraphNode_GridSpatialization2D::PrepareForReplication()
 {
 	RG_QUICK_SCOPE_CYCLE_COUNTER(UReplicationGraphNode_GridSpatialization2D_PrepareForReplication);
@@ -4509,121 +4776,178 @@ void UReplicationGraphNode_GridSpatialization2D::PrepareForReplication()
 	}
 }
 
-void UReplicationGraphNode_GridSpatialization2D::GatherActorListsForConnection(const FConnectionGatherActorListParameters& Params)
+// Small structure to make it easier to keep track of 
+// information regarding current players for a connection when working with grids
+struct FPlayerGridCellInformation
 {
-	if (Params.Viewer.ViewLocation.Z > ConnectionMaxZ)
+	FPlayerGridCellInformation(UNetConnection* InConnection, FIntPoint InCurLocation) :
+		Connection(InConnection), CurLocation(InCurLocation), PrevLocation(FIntPoint::ZeroValue)
 	{
-		return;
 	}
 
-	FVector ClampedViewLoc = Params.Viewer.ViewLocation;
-
-	if (GridBounds.IsValid)
-	{
-		ClampedViewLoc = GridBounds.GetClosestPointTo(ClampedViewLoc);
-	}
-
-	// Find out what bucket the view is in
-
-	int32 CellX = (ClampedViewLoc.X - SpatialBias.X) / CellSize;
-	if (CellX < 0)
-	{
-		UE_LOG(LogReplicationGraph, Log, TEXT("Net view location.X %s is less than the spatial bias %s"), *ClampedViewLoc.ToString(), *SpatialBias.ToString());
-		CellX = 0;
-	}
-	
-	TArray<UReplicationGraphNode_GridCell*>& GridX = GetGridX(CellX);
-
-	// -----------
+	UNetConnection* Connection;
+	FIntPoint CurLocation;
+	FIntPoint PrevLocation;
+};
 
-	int32 CellY = (ClampedViewLoc.Y - SpatialBias.Y) / CellSize;
-	if (CellY < 0)
-	{
-		UE_LOG(LogReplicationGraph, Log, TEXT("Net view location.Y %s is less than the spatial bias %s"), *ClampedViewLoc.ToString(), *SpatialBias.ToString());
-		CellY = 0;
-	}
-	if (GridX.Num() <= CellY)
-	{
-		GridX.SetNum(CellY+1);
-	}
+void UReplicationGraphNode_GridSpatialization2D::GatherActorListsForConnection(const FConnectionGatherActorListParameters& Params)
+{
+	TArray<FLastLocationGatherInfo, FReplicationGraphConnectionsAllocator>& LastLocationArray = Params.ConnectionManager.LastGatherLocations;
+	TArray<FVector2D, FReplicationGraphConnectionsAllocator> UniqueCurrentLocations;
 
-	UReplicationGraphNode_GridCell* CellNode = GridX[CellY];
-	if (CellNode)
+	// Consider all users that are in cells for this connection. 
+	// From here, generate a list of coordinates, we'll later work through each coordinate pairing
+	// to find the cells that are actually active. This reduces redundancy and cache misses.
+	TArray<FPlayerGridCellInformation, FReplicationGraphConnectionsAllocator> ActiveGridCells;
+	for (const FNetViewer& CurViewer : Params.Viewers)
 	{
-		CellNode->GatherActorListsForConnection(Params);
-	}
+		if (CurViewer.ViewLocation.Z > ConnectionMaxZ || CurViewer.Connection == nullptr)
+		{
+			continue;
+		}
 
-	if (bDestroyDormantDynamicActors && CVar_RepGraph_DormantDynamicActorsDestruction > 0 )
-	{
-		FVector LastGatherLoc = Params.ConnectionManager.LastGatherLocation;
-		
+		// Figure out positioning
+		FVector ClampedViewLoc = CurViewer.ViewLocation;
 		if (GridBounds.IsValid)
 		{
-			LastGatherLoc = GridBounds.GetClosestPointTo(LastGatherLoc);
+			ClampedViewLoc = GridBounds.GetClosestPointTo(ClampedViewLoc);
 		}
 
-		int32 PrevX = FMath::Max(0, (int32)((LastGatherLoc.X - SpatialBias.X) / CellSize));
-		int32 PrevY = FMath::Max(0, (int32)((LastGatherLoc.Y - SpatialBias.Y) / CellSize));
+		// Find out what bucket the view is in
+		int32 CellX = (ClampedViewLoc.X - SpatialBias.X) / CellSize;
+		if (CellX < 0)
+		{
+			UE_LOG(LogReplicationGraph, Log, TEXT("Net view location.X %s is less than the spatial bias %s"), *ClampedViewLoc.ToString(), *SpatialBias.ToString());
+			CellX = 0;
+		}
 
-		// if the grid cell changed this gather
-		if ((CellX != PrevX) || (CellY != PrevY))
+		int32 CellY = (ClampedViewLoc.Y - SpatialBias.Y) / CellSize;
+		if (CellY < 0)
 		{
-			UReplicationGraphNode_GridCell* PrevCell = GetCell(PrevX, PrevY);
-			DoDormantDynamicActorDestruction(Params, CellNode, PrevCell);
+			UE_LOG(LogReplicationGraph, Log, TEXT("Net view location.Y %s is less than the spatial bias %s"), *ClampedViewLoc.ToString(), *SpatialBias.ToString());
+			CellY = 0;
 		}
-	}
-}
 
-void UReplicationGraphNode_GridSpatialization2D::DoDormantDynamicActorDestruction(const FConnectionGatherActorListParameters& Params, UReplicationGraphNode_GridCell* CurrentCell, UReplicationGraphNode_GridCell* PreviousCell)
-{
-	RG_QUICK_SCOPE_CYCLE_COUNTER(UReplicationGraphNode_GridSpatialization2D_DoDormantDynamicActorDestruction);
+		// Save this information out for later.
+		FPlayerGridCellInformation NewPlayerCell(CurViewer.Connection, FIntPoint(CellX, CellY));
+		FLastLocationGatherInfo* GatherInfoForConnection = LastLocationArray.FindByKey<UNetConnection*>(CurViewer.Connection);
 
-	FActorRepListRefView DormantActorList;
-	FActorRepListRefView PrevDormantActorList;
+		// Add any missing last location information that we don't have
+		if (GatherInfoForConnection == nullptr)
+		{
+			GatherInfoForConnection = &LastLocationArray[LastLocationArray.Add(FLastLocationGatherInfo(CurViewer.Connection, FVector(ForceInitToZero)))];
+		}
 
-	if (CurrentCell)
-	{
-		if (UReplicationGraphNode_DormancyNode* DormancyNode = CurrentCell->GetDormancyNode())
+		// Clean up the location data for this connection to be grid bound
+		FVector LastLocationForConnection = GatherInfoForConnection->LastLocation;
+		if (GridBounds.IsValid)
 		{
-			DormancyNode->ConditionalGatherDormantDynamicActors(DormantActorList, Params, nullptr);
+			LastLocationForConnection = GridBounds.GetClosestPointTo(LastLocationForConnection);
 		}
-	}
 
-	if (PreviousCell)
-	{
-		if (UReplicationGraphNode_DormancyNode* DormancyNode = PreviousCell->GetDormancyNode())
+		// Try to determine the previous location of the user.
+		NewPlayerCell.PrevLocation.X = FMath::Max(0, (int32)((LastLocationForConnection.X - SpatialBias.X) / CellSize));
+		NewPlayerCell.PrevLocation.Y = FMath::Max(0, (int32)((LastLocationForConnection.Y - SpatialBias.Y) / CellSize));
+
+		// If we have not operated on this cell yet (meaning it's not shared by anyone else), gather for it.
+		if (!UniqueCurrentLocations.Contains(NewPlayerCell.CurLocation))
 		{
-			DormancyNode->ConditionalGatherDormantDynamicActors(PrevDormantActorList, Params, &DormantActorList);
+			TArray<UReplicationGraphNode_GridCell*>& GridX = GetGridX(CellX);
+			if (GridX.Num() <= CellY)
+			{
+				GridX.SetNum(CellY + 1);
+			}
+
+			UReplicationGraphNode_GridCell* CellNode = GridX[CellY];
+			if (CellNode)
+			{
+				CellNode->GatherActorListsForConnection(Params);
+			}
+
+			UniqueCurrentLocations.Add(NewPlayerCell.CurLocation);
 		}
-	}
 
-	// any previous dormant actors not in the current node dormant list
-	if (PrevDormantActorList.IsValid() == false)
-	{
-		return;
+		// Add this to things we consider later.
+		ActiveGridCells.Add(NewPlayerCell);
 	}
 
-	for (FActorRepListType& Actor : PrevDormantActorList)
+	if (bDestroyDormantDynamicActors && CVar_RepGraph_DormantDynamicActorsDestruction > 0)
 	{
-		Params.ConnectionManager.NotifyAddDormantDestructionInfo(Actor);
+		bool bCellHasChanged = false;
+		FActorRepListRefView PrevDormantActorList;
 
-		if (FConnectionReplicationActorInfo* ActorInfo = Params.ConnectionManager.ActorInfoMap.Find(Actor))
+		// Process and create the dormancy list for the active grid for this user
+		for (const FPlayerGridCellInformation& CellInfo : ActiveGridCells)
 		{
-			ActorInfo->bDormantOnConnection = false;
+			const int32& CellX = CellInfo.CurLocation.X;
+			const int32& CellY = CellInfo.CurLocation.Y;
+			const int32& PrevX = CellInfo.PrevLocation.X;
+			const int32& PrevY = CellInfo.PrevLocation.Y;
+
+			// The idea is that if the previous location is a current location for any other user, we do not bother to do operations on this cell
+			// However, if the current location matches with a current location of another user, continue anyways.
+			//
+			// as above, if the grid cell changed this gather and is not in current use by any other viewer
+
+			// TODO: There is a potential list gathering redundancy if two actors share the same current and previous cell information
+			// but this should just result in a wasted cycle if anything.
+			if (((CellX != PrevX) || (CellY != PrevY)) && !UniqueCurrentLocations.Contains(CellInfo.PrevLocation))
+			{
+				bCellHasChanged = true;
+				RG_QUICK_SCOPE_CYCLE_COUNTER(UReplicationGraphNode_GridSpatialization2D_CellChangeDormantRelevancy);
+				FActorRepListRefView DormantActorList;
 
-						// add back to connection specific dormancy nodes
-			const FActorCellInfo CellInfo = GetCellInfoForActor(Actor, Actor->GetActorLocation(), ActorInfo->GetCullDistance());
+				TArray<UReplicationGraphNode_GridCell*>& GridX = GetGridX(CellX);
+				UReplicationGraphNode_GridCell* CellNode = GridX[CellY];
 
-			GetGridNodesForActor(Actor, CellInfo, GatheredNodes);
+				if (CellNode)
+				{
+					if (UReplicationGraphNode_DormancyNode* DormancyNode = CellNode->GetDormancyNode())
+					{
+						DormancyNode->ConditionalGatherDormantDynamicActors(DormantActorList, Params, nullptr);
+					}
+				}
 
-			for (UReplicationGraphNode_GridCell* Node : GatheredNodes)
+				// Determine dormant actors for our last location. Do not add actors if they are relevant to anyone.
+				if (UReplicationGraphNode_GridCell* PrevCell = GetCell(GetGridX(PrevX), PrevY))
+				{
+					if (UReplicationGraphNode_DormancyNode* DormancyNode = PrevCell->GetDormancyNode())
+					{
+						DormancyNode->ConditionalGatherDormantDynamicActors(PrevDormantActorList, Params, &DormantActorList, true);
+					}
+				}
+			}
+		}
+
+		// Now process the previous dormant list to handle destruction
+		if (bCellHasChanged && PrevDormantActorList.IsValid())
+		{
+			// any previous dormant actors not in the current node dormant list
+			for (FActorRepListType& Actor : PrevDormantActorList)
 			{
-				if (UReplicationGraphNode_DormancyNode* DormancyNode = Node->GetDormancyNode())
+				Params.ConnectionManager.NotifyAddDormantDestructionInfo(Actor);
+
+				if (FConnectionReplicationActorInfo* ActorInfo = Params.ConnectionManager.ActorInfoMap.Find(Actor))
 				{
-					// Only notify the connection node if this client was previously inside the cell.
-					if (UReplicationGraphNode_ConnectionDormanyNode* ConnectionDormancyNode = DormancyNode->GetExistingConnectionNode(Params))
+					ActorInfo->bDormantOnConnection = false;
+					// Ideally, no actor info outside this list should be set to true, so we don't have to worry about resetting them.
+					// However we could consider iterating through the actor map to reset all of them.
+					ActorInfo->bGridSpatilization_AlreadyDormant = false;
+
+					// add back to connection specific dormancy nodes
+					const FActorCellInfo CellInfo = GetCellInfoForActor(Actor, Actor->GetActorLocation(), ActorInfo->GetCullDistance());
+					GetGridNodesForActor(Actor, CellInfo, GatheredNodes);
+
+					for (UReplicationGraphNode_GridCell* Node : GatheredNodes)
 					{
-						ConnectionDormancyNode->NotifyActorDormancyFlush(Actor);
+						if (UReplicationGraphNode_DormancyNode* DormancyNode = Node->GetDormancyNode())
+						{
+							// Only notify the connection node if this client was previously inside the cell.
+							if (UReplicationGraphNode_ConnectionDormanyNode* ConnectionDormancyNode = DormancyNode->GetExistingConnectionNode(Params))
+							{
+								ConnectionDormancyNode->NotifyActorDormancyFlush(Actor);
+							}
+						}
 					}
 				}
 			}
@@ -4847,8 +5171,34 @@ void UReplicationGraphNode_AlwaysRelevant_ForConnection::GatherActorListsForConn
 
 	// Reset and rebuild another list that will contains our current viewer/viewtarget
 	ReplicationActorList.Reset();
-	UpdateActor(Params.Viewer.InViewer, LastViewer);
-	UpdateActor(Params.Viewer.ViewTarget, LastViewTarget);
+
+	for (const FNetViewer& CurViewer : Params.Viewers)
+	{
+		if (CurViewer.Connection == nullptr)
+		{
+			continue;
+		}
+
+		FAlwaysRelevantActorInfo* LastData = PastRelevantActors.FindByKey<UNetConnection*>(CurViewer.Connection);
+
+		// We've not seen this actor before, go ahead and add them.
+		if (LastData == nullptr)
+		{
+			FAlwaysRelevantActorInfo NewActorInfo;
+			NewActorInfo.Connection = CurViewer.Connection;
+			LastData = &(PastRelevantActors[PastRelevantActors.Add(NewActorInfo)]);
+		}
+
+		check(LastData != nullptr);
+
+		UpdateActor(CurViewer.InViewer, LastData->LastViewer);
+		UpdateActor(CurViewer.ViewTarget, LastData->LastViewTarget);
+	}
+
+	// Remove excess
+	PastRelevantActors.RemoveAll([&](FAlwaysRelevantActorInfo& RelActorInfo){
+		return RelActorInfo.Connection == nullptr;
+	});
 	
 	if (ReplicationActorList.Num() > 0)
 	{
diff --git a/Engine/Plugins/Runtime/ReplicationGraph/Source/Private/ReplicationGraphDebugging.cpp b/Engine/Plugins/Runtime/ReplicationGraph/Source/Private/ReplicationGraphDebugging.cpp
--- a/Engine/Plugins/Runtime/ReplicationGraph/Source/Private/ReplicationGraphDebugging.cpp
+++ b/Engine/Plugins/Runtime/ReplicationGraph/Source/Private/ReplicationGraphDebugging.cpp
@@ -333,44 +333,65 @@ bool AReplicationGraphDebugActor::ServerCellInfo_Validate()
 
 void AReplicationGraphDebugActor::ServerCellInfo_Implementation()
 {
-	FNetViewer Viewer(GetNetConnection(), 0.f);
+	TArray<FVector, FReplicationGraphConnectionsAllocator> LocationsSent;
+	TArray<UNetConnection*, FReplicationGraphConnectionsAllocator> ConnectionsToConsider;
 
-	UReplicationGraphNode_GridSpatialization2D* GridNode =  nullptr;
-	for (UReplicationGraphNode* Node : ReplicationGraph->GlobalGraphNodes)
+	UNetConnection* PrimaryNetConnection = GetNetConnection();
+	ConnectionsToConsider.Add(PrimaryNetConnection);
+	for (UChildConnection* Child : PrimaryNetConnection->Children)
 	{
-		GridNode = Cast<UReplicationGraphNode_GridSpatialization2D>(Node);
-		if (GridNode)
-		{
-			break;
-	}
+		ConnectionsToConsider.Add((UNetConnection*)Child);
 	}
 
-	if (GridNode == nullptr)
+	for (UNetConnection* Connection : ConnectionsToConsider)
 	{
-		return;
-	}
+		FNetViewer Viewer(Connection, 0.f);
 
-	int32 CellX = FMath::Max<int32>(0, (Viewer.ViewLocation.X - GridNode->SpatialBias.X) / GridNode->CellSize);
-	int32 CellY = FMath::Max<int32>(0, (Viewer.ViewLocation.Y - GridNode->SpatialBias.Y) / GridNode->CellSize);
+		UReplicationGraphNode_GridSpatialization2D* GridNode = nullptr;
+		for (UReplicationGraphNode* Node : ReplicationGraph->GlobalGraphNodes)
+		{
+			GridNode = Cast<UReplicationGraphNode_GridSpatialization2D>(Node);
+			if (GridNode)
+			{
+				break;
+			}
+		}
 
-	TArray<AActor*> ActorsInCell;
+		if (GridNode == nullptr)
+		{
+			return;
+		}
 
-	FVector CellLocation(GridNode->SpatialBias.X + (((float)(CellX)+0.5f) * GridNode->CellSize), GridNode->SpatialBias.Y + (((float)(CellY)+0.5f) * GridNode->CellSize), Viewer.ViewLocation.Z);
-	FVector CellExtent(GridNode->CellSize, GridNode->CellSize, 10.f);
+		int32 CellX = FMath::Max<int32>(0, (Viewer.ViewLocation.X - GridNode->SpatialBias.X) / GridNode->CellSize);
+		int32 CellY = FMath::Max<int32>(0, (Viewer.ViewLocation.Y - GridNode->SpatialBias.Y) / GridNode->CellSize);
 
-	if (GridNode->Grid.IsValidIndex(CellX))
-	{
-		TArray<UReplicationGraphNode_GridCell*>& GridY = GridNode->Grid[CellX];
-		if (GridY.IsValidIndex(CellY))
+		FVector CellLocation(GridNode->SpatialBias.X + (((float)(CellX)+0.5f) * GridNode->CellSize), GridNode->SpatialBias.Y + (((float)(CellY)+0.5f) * GridNode->CellSize), Viewer.ViewLocation.Z);
+
+		if (LocationsSent.Contains(CellLocation))
 		{
-			if (UReplicationGraphNode_GridCell* LeafNode = GridY[CellY])
+			UE_LOG(LogReplicationGraph, Verbose, TEXT("Skipping location %s as we've already sent it"), *(CellLocation.ToString()));
+			continue;
+		}
+
+		LocationsSent.Add(CellLocation);
+
+		TArray<AActor*> ActorsInCell;
+		FVector CellExtent(GridNode->CellSize, GridNode->CellSize, 10.f);
+
+		if (GridNode->Grid.IsValidIndex(CellX))
+		{
+			TArray<UReplicationGraphNode_GridCell*>& GridY = GridNode->Grid[CellX];
+			if (GridY.IsValidIndex(CellY))
 			{
-				LeafNode->GetAllActorsInNode_Debugging(ActorsInCell);
+				if (UReplicationGraphNode_GridCell* LeafNode = GridY[CellY])
+				{
+					LeafNode->GetAllActorsInNode_Debugging(ActorsInCell);
+				}
 			}
 		}
-	}
 
-	ClientCellInfo(CellLocation, CellExtent, ActorsInCell);
+		ClientCellInfo(CellLocation, CellExtent, ActorsInCell);
+	}	
 }
 
 void AReplicationGraphDebugActor::ClientCellInfo_Implementation(FVector CellLocation, FVector CellExtent, const TArray<AActor*>& Actors)
@@ -869,7 +890,6 @@ void LogGraphHelper(FOutputDevice& Ar, const TArray< FString >& Args)
 
 	if (!Graph)
 	{
-		UE_LOG(LogReplicationGraph, Warning, TEXT("Could not find valid Replication Graph."));
 		return;
 	}
 
@@ -1108,38 +1128,42 @@ void PrintPrioritizedList(FOutputDevice& Ar, UNetReplicationGraphConnection* Con
 	
 	// Skipped actors
 #if REPGRAPH_DETAILS
-	Ar.Logf(TEXT("[%d Skipped Actors]"), PrioritizedList->SkippedDebugDetails->Num());
-
-	FNativeClassAccumulator DormantClasses;
-	FNativeClassAccumulator CulledClasses;
 
-	for (const FSkippedActorFullDebugDetails& SkippedDetails : *PrioritizedList->SkippedDebugDetails)
+	if (PrioritizedList->SkippedDebugDetails.IsValid())
 	{
-		FString SkippedStr;
-		if (SkippedDetails.bWasDormant)
-		{
-			SkippedStr = TEXT("Dormant");
-			DormantClasses.Increment(SkippedDetails.Actor->GetClass());
-		}
-		else if (SkippedDetails.DistanceCulled > 0.f)
-		{
-			SkippedStr = FString::Printf(TEXT("Dist Culled %.2f"), SkippedDetails.DistanceCulled);
-			CulledClasses.Increment(SkippedDetails.Actor->GetClass());
-		}
-		else if (SkippedDetails.FramesTillNextReplication > 0)
-		{
-			SkippedStr = FString::Printf(TEXT("Not ready (%d frames left)"), SkippedDetails.FramesTillNextReplication);
-		}
-		else
+		Ar.Logf(TEXT("[%d Skipped Actors]"), PrioritizedList->SkippedDebugDetails->Num());
+
+		FNativeClassAccumulator DormantClasses;
+		FNativeClassAccumulator CulledClasses;
+
+		for (const FSkippedActorFullDebugDetails& SkippedDetails : *PrioritizedList->SkippedDebugDetails)
 		{
-			SkippedStr = TEXT("Unknown???");
-		}
+			FString SkippedStr;
+			if (SkippedDetails.bWasDormant)
+			{
+				SkippedStr = TEXT("Dormant");
+				DormantClasses.Increment(SkippedDetails.Actor->GetClass());
+			}
+			else if (SkippedDetails.DistanceCulled > 0.f)
+			{
+				SkippedStr = FString::Printf(TEXT("Dist Culled %.2f"), SkippedDetails.DistanceCulled);
+				CulledClasses.Increment(SkippedDetails.Actor->GetClass());
+			}
+			else if (SkippedDetails.FramesTillNextReplication > 0)
+			{
+				SkippedStr = FString::Printf(TEXT("Not ready (%d frames left)"), SkippedDetails.FramesTillNextReplication);
+			}
+			else
+			{
+				SkippedStr = TEXT("Unknown???");
+			}
 
-		Ar.Logf(TEXT("%-40s %s"), *GetActorRepListTypeDebugString(SkippedDetails.Actor), *SkippedStr);
-	}
+			Ar.Logf(TEXT("%-40s %s"), *GetActorRepListTypeDebugString(SkippedDetails.Actor), *SkippedStr);
+		}
 
 		Ar.Logf(TEXT(" Dormant Classes: %s"), *DormantClasses.BuildString());
 		Ar.Logf(TEXT(" Culled Classes: %s"), *CulledClasses.BuildString());
+	}
 
 #endif
 
diff --git a/Engine/Plugins/Runtime/ReplicationGraph/Source/Public/BasicReplicationGraph.h b/Engine/Plugins/Runtime/ReplicationGraph/Source/Public/BasicReplicationGraph.h
--- a/Engine/Plugins/Runtime/ReplicationGraph/Source/Public/BasicReplicationGraph.h
+++ b/Engine/Plugins/Runtime/ReplicationGraph/Source/Public/BasicReplicationGraph.h
@@ -12,7 +12,7 @@ struct FConnectionAlwaysRelevantNodePair
 	GENERATED_BODY()
 	FConnectionAlwaysRelevantNodePair() { }
 	FConnectionAlwaysRelevantNodePair(UNetConnection* InConnection, UReplicationGraphNode_AlwaysRelevant_ForConnection* InNode) : NetConnection(InConnection), Node(InNode) { }
-	bool operator==(const UNetConnection* InConnection) const { return InConnection == NetConnection; }
+	bool operator==(UNetConnection* InConnection) const;
 
 	UPROPERTY()
 	UNetConnection* NetConnection = nullptr;
diff --git a/Engine/Plugins/Runtime/ReplicationGraph/Source/Public/ReplicationGraph.h b/Engine/Plugins/Runtime/ReplicationGraph/Source/Public/ReplicationGraph.h
--- a/Engine/Plugins/Runtime/ReplicationGraph/Source/Public/ReplicationGraph.h
+++ b/Engine/Plugins/Runtime/ReplicationGraph/Source/Public/ReplicationGraph.h
@@ -367,6 +367,7 @@ protected:
 	virtual void GatherActors(const FActorRepListRefView& RepList, FGlobalActorReplicationInfoMap& GlobalMap, FPerConnectionActorInfoMap& ConnectionMap, const FConnectionGatherActorListParameters& Params, UNetConnection* NetConnection);
 	virtual void GatherActors_DistanceOnly(const FActorRepListRefView& RepList, FGlobalActorReplicationInfoMap& GlobalMap, FPerConnectionActorInfoMap& ConnectionMap, const FConnectionGatherActorListParameters& Params);
 
+	void CalcFrequencyForActor(AActor* Actor, UReplicationGraph* RepGraph, UNetConnection* NetConnection, FGlobalActorReplicationInfo& GlobalInfo, FConnectionReplicationActorInfo& ConnectionInfo, FSettings& MySettings, const FNetViewerArray& Viewers, const uint32 FrameNum, int32 ExistingItemIndex);
 	void CalcFrequencyForActor(AActor* Actor, UReplicationGraph* RepGraph, UNetConnection* NetConnection, FGlobalActorReplicationInfo& GlobalInfo, FConnectionReplicationActorInfo& ConnectionInfo, FSettings& MySettings, const FVector& ConnectionViewLocation, const FVector& ConnectionViewDir, const uint32 FrameNum, int32 ExistingItemIndex);
 };
 
@@ -418,7 +419,7 @@ public:
 
 	void OnActorDormancyFlush(FActorRepListType Actor, FGlobalActorReplicationInfo& GlobalInfo);
 
-	void ConditionalGatherDormantDynamicActors(FActorRepListRefView& RepList, const FConnectionGatherActorListParameters& Params, FActorRepListRefView* RemovedList);
+	void ConditionalGatherDormantDynamicActors(FActorRepListRefView& RepList, const FConnectionGatherActorListParameters& Params, FActorRepListRefView* RemovedList, bool bEnforceReplistUniqueness=false);
 
 	UReplicationGraphNode_ConnectionDormanyNode* GetExistingConnectionNode(const FConnectionGatherActorListParameters& Params);
 
@@ -587,8 +588,6 @@ private:
 	
 	void PutStaticActorIntoCell(const FNewReplicatedActorInfo& ActorInfo, FGlobalActorReplicationInfo& ActorRepInfo, bool bDormancyDriven);
 
-	void DoDormantDynamicActorDestruction(const FConnectionGatherActorListParameters& Params, UReplicationGraphNode_GridCell* CurrentCell, UReplicationGraphNode_GridCell* PreviousCell);
-
 	UReplicationGraphNode_GridCell* GetCellNode(UReplicationGraphNode_GridCell*& NodePtr)
 	{
 		if (NodePtr == nullptr)
@@ -676,6 +675,26 @@ protected:
 
 // -----------------------------------
 
+USTRUCT()
+struct FAlwaysRelevantActorInfo
+{
+	GENERATED_BODY()
+
+	UPROPERTY()
+	const UNetConnection* Connection = nullptr;
+
+	UPROPERTY()
+	AActor* LastViewer = nullptr;
+
+	UPROPERTY()
+	AActor* LastViewTarget = nullptr;
+
+	bool operator==(UNetConnection* Other) const
+	{
+		return Connection == Other;
+	}
+};
+
 /** Adds actors that are always relevant for a connection. This engine version just adds the PlayerController and ViewTarget (usually the pawn) */
 UCLASS()
 class REPLICATIONGRAPH_API UReplicationGraphNode_AlwaysRelevant_ForConnection : public UReplicationGraphNode_ActorList
@@ -689,6 +708,9 @@ public:
 	/** Rebuilt-every-frame list based on UNetConnection state */
 	FActorRepListRefView ReplicationActorList;
 
+	/** List of previously (or currently if nothing changed last tick) focused actor data per connection */
+	TArray<FAlwaysRelevantActorInfo, FReplicationGraphConnectionsAllocator> PastRelevantActors;
+
 	UPROPERTY()
 	AActor* LastViewer = nullptr;
 	
@@ -922,9 +944,11 @@ protected:
 #endif
 
 	/** Default Replication Path */
+	void ReplicateActorListsForConnections_Default(UNetReplicationGraphConnection* ConnectionManager, FGatheredReplicationActorLists& GatheredReplicationListsForConnection, FNetViewerArray& Viewers);
 	void ReplicateActorListsForConnection_Default(UNetReplicationGraphConnection* ConnectionManager, FGatheredReplicationActorLists& GatheredReplicationListsForConnection, FNetViewer& Viewer);
 
 	/** "FastShared" Replication Path */
+	void ReplicateActorListsForConnections_FastShared(UNetReplicationGraphConnection* ConnectionManager, FGatheredReplicationActorLists& GatheredReplicationListsForConnection, FNetViewerArray& Viewers);
 	void ReplicateActorListsForConnection_FastShared(UNetReplicationGraphConnection* ConnectionManager, FGatheredReplicationActorLists& GatheredReplicationListsForConnection, FNetViewer& Viewer);
 
 	/** Connections needing a FlushNet in PostTickDispatch */
@@ -950,6 +974,25 @@ private:
 // --------------------------------------------------------------------------------------------------------------------------------------------
 // --------------------------------------------------------------------------------------------------------------------------------------------
 
+USTRUCT()
+struct FLastLocationGatherInfo
+{
+	GENERATED_BODY()
+
+	FLastLocationGatherInfo() : Connection(nullptr), LastLocation(FVector(ForceInitToZero)) {}
+	FLastLocationGatherInfo(const UNetConnection* InConnection, FVector InLastLocation) : Connection(InConnection), LastLocation(InLastLocation) {}
+
+	UPROPERTY()
+	const UNetConnection* Connection;
+
+	UPROPERTY()
+	FVector LastLocation;
+
+	bool operator==(UNetConnection* Other) const
+	{
+		return Connection == Other;
+	}
+};
 
 /** Manages actor replication for a specific connection */
 UCLASS(transient)
@@ -992,6 +1035,8 @@ public:
 
 	FVector LastGatherLocation;
 
+	TArray<FLastLocationGatherInfo, FReplicationGraphConnectionsAllocator> LastGatherLocations;
+
 	// Nb of bits sent for actor channel creation when a dedicated budget is allocated for this
 	int32 QueuedBitsForActorDiscovery = 0;
 
@@ -1024,6 +1069,9 @@ public:
 	virtual void NotifyResetDestructionInfo() override;
 	//~ End UReplicationConnectionDriver Interface
 
+	/** Generates a set of all the visible level names for this connection and its subconnections (if any) */
+	virtual void GetClientVisibleLevelNames(TSet<FName>& OutLevelNames) const;
+
 private:
 
 	friend UReplicationGraph;
@@ -1044,7 +1092,7 @@ private:
 
 	bool PrepareForReplication();
 
-	int64 ReplicateDestructionInfos(const FVector& ConnectionViewLocation, const float DestructInfoMaxDistanceSquared);
+	int64 ReplicateDestructionInfos(const FNetViewerArray& InViewers, const float DestructInfoMaxDistanceSquared);
 	
 	int64 ReplicateDormantDestructionInfos();
 
diff --git a/Engine/Plugins/Runtime/ReplicationGraph/Source/Public/ReplicationGraphTypes.h b/Engine/Plugins/Runtime/ReplicationGraph/Source/Public/ReplicationGraphTypes.h
--- a/Engine/Plugins/Runtime/ReplicationGraph/Source/Public/ReplicationGraphTypes.h
+++ b/Engine/Plugins/Runtime/ReplicationGraph/Source/Public/ReplicationGraphTypes.h
@@ -787,9 +787,9 @@ private:
 /** Per-Actor data that is stored per connection */
 struct FConnectionReplicationActorInfo
 {
-	FConnectionReplicationActorInfo() : bDormantOnConnection(0), bTearOff(0)  { }
+	FConnectionReplicationActorInfo() : bDormantOnConnection(0), bTearOff(0), bGridSpatilization_AlreadyDormant(0) { }
 
-	FConnectionReplicationActorInfo(const FGlobalActorReplicationInfo& GlobalInfo) : bDormantOnConnection(0), bTearOff(0)
+	FConnectionReplicationActorInfo(const FGlobalActorReplicationInfo& GlobalInfo) : bDormantOnConnection(0), bTearOff(0), bGridSpatilization_AlreadyDormant(0)
 	{
 		// Pull data from the global actor info. This is done for things that we just want to duplicate in both places so that we can avoid a lookup into the global map
 		// and also for things that we want to be overridden per (connection/actor)
@@ -846,6 +846,9 @@ public:
 	uint8 bDormantOnConnection:1;
 	uint8 bTearOff:1;
 
+	/** Used as an optimization when doing 2D Grid Spatilization, prevents replicating the dormancy of the same actor twice in splitscreen evaluations. */
+	uint8 bGridSpatilization_AlreadyDormant:1;
+
 	void LogDebugString(FOutputDevice& Ar) const;
 };
 
@@ -1119,6 +1122,13 @@ private:
 // Connection Gather Actor List Parameters
 // --------------------------------------------------------------------------------------------------------------------------------------------
 // --------------------------------------------------------------------------------------------------------------------------------------------
+#if !defined(REPGRAPH_VIEWERS_PER_CONNECTION)
+#define REPGRAPH_VIEWERS_PER_CONNECTION 4
+#endif
+
+// Allocator for a connection and all of its sub connections
+typedef TInlineAllocator<REPGRAPH_VIEWERS_PER_CONNECTION> FReplicationGraphConnectionsAllocator;
+typedef TArray<FNetViewer, FReplicationGraphConnectionsAllocator> FNetViewerArray;
 
 // Parameter structure for what we actually pass down during the Gather phase.
 struct FConnectionGatherActorListParameters
@@ -1126,10 +1136,19 @@ struct FConnectionGatherActorListParameters
 	FConnectionGatherActorListParameters(FNetViewer& InViewer, UNetReplicationGraphConnection& InConnectionManager, TSet<FName>& InClientVisibleLevelNamesRef, uint32 InReplicationFrameNum, FGatheredReplicationActorLists& InOutGatheredReplicationLists)
 		: Viewer(InViewer), ConnectionManager(InConnectionManager), ReplicationFrameNum(InReplicationFrameNum), OutGatheredReplicationLists(InOutGatheredReplicationLists), ClientVisibleLevelNamesRef(InClientVisibleLevelNamesRef)
 	{
+		Viewers.Add(InViewer);
 	}
 
+	FConnectionGatherActorListParameters(FNetViewerArray& InViewers, UNetReplicationGraphConnection& InConnectionManager, TSet<FName>& InClientVisibleLevelNamesRef, uint32 InReplicationFrameNum, FGatheredReplicationActorLists& InOutGatheredReplicationLists)
+		: Viewer(InViewers[0]), Viewers(InViewers), ConnectionManager(InConnectionManager), ReplicationFrameNum(InReplicationFrameNum), OutGatheredReplicationLists(InOutGatheredReplicationLists), ClientVisibleLevelNamesRef(InClientVisibleLevelNamesRef)
+	{
+	}
+
+
 	/** In: The Data the nodes have to work with */
 	FNetViewer& Viewer;
+
+	FNetViewerArray Viewers;
 	UNetReplicationGraphConnection& ConnectionManager;
 	uint32 ReplicationFrameNum;
 
@@ -1154,14 +1173,13 @@ struct FConnectionGatherActorListParameters
 
 
 	// Cached off reference for fast Level Visibility lookup
-	TSet<FName>& ClientVisibleLevelNamesRef;
+	const TSet<FName>& ClientVisibleLevelNamesRef;
 
 private:
 
 	mutable FName LastCheckedVisibleLevelName;
 };
 
-
 // --------------------------------------------------------------------------------------------------------------------------------------------
 // --------------------------------------------------------------------------------------------------------------------------------------------
 // NewReplicatedActorInfo
