PR #5705: Correct Typo: Dormany => Dormancy (Contributed by [REDACTED_PERSON])

Renamed UReplicationGraphNode_ConnectionDormanyNode to UReplicationGraphNode_ConnectionDormancyNode to correct the typo in name

[RN] deprecation Engine/Network

diff --git a/Engine/Plugins/Runtime/ReplicationGraph/Source/Private/ReplicationGraph.cpp b/Engine/Plugins/Runtime/ReplicationGraph/Source/Private/ReplicationGraph.cpp
--- a/Engine/Plugins/Runtime/ReplicationGraph/Source/Private/ReplicationGraph.cpp
+++ b/Engine/Plugins/Runtime/ReplicationGraph/Source/Private/ReplicationGraph.cpp
@@ -101,7 +101,7 @@ REPGRAPH_DEVCVAR_SHIPCONST(int32, "Net.RepGraph.LogActorAdd", CVar_RepGraph_LogA
 REPGRAPH_DEVCVAR_SHIPCONST(int32, "Net.RepGraph.Verify", CVar_RepGraph_Verify, 0, "Additional, slow, verification is done on replication graph nodes. Guards against: invalid actors and dupes");
 REPGRAPH_DEVCVAR_SHIPCONST(int32, "Net.RepGraph.DisableBandwithLimit", CVar_RepGraph_DisableBandwithLimit, 0, "Disables the IsNetReady() check, effectively replicating all actors that want to replicate to each connection.");
 
-REPGRAPH_DEVCVAR_SHIPCONST(int32, "Net.RepGraph.TrickleDistCullOnDormanyNodes", CVar_RepGraph_TrickleDistCullOnDormanyNodes, 1, "Actors in a dormancy node that are distance culled will trickle through as dormancy node empties");
+REPGRAPH_DEVCVAR_SHIPCONST(int32, "Net.RepGraph.TrickleDistCullOnDormancyNodes", CVar_RepGraph_TrickleDistCullOnDormancyNodes, 1, "Actors in a dormancy node that are distance culled will trickle through as dormancy node empties");
 REPGRAPH_DEVCVAR_SHIPCONST(int32, "Net.RepGraph.EnableRPCSendPolicy", CVar_RepGraph_EnableRPCSendPolicy, 1, "Enables RPC send policy (e.g, force certain functions to send immediately rather than be queued)");
 REPGRAPH_DEVCVAR_SHIPCONST(int32, "Net.RepGraph.EnableFastSharedPath", CVar_RepGraph_EnableFastSharedPath, 1, "Enables FastShared replication path for lists with EActorRepListTypeFlags::FastShared flag");
 REPGRAPH_DEVCVAR_SHIPCONST(int32, "Net.RepGraph.EnableDynamicAllocationWarnings", CVar_RepGraph_EnableDynamicAllocationWarnings, 1, "Enables debug information whenever RepGraph needs to allocate new Actor Lists.");
@@ -3504,7 +3504,7 @@ void UReplicationGraphNode_DynamicSpatialFrequency::GatherActors_DistanceOnly(co
 
 // --------------------------------------------------------------------------------------------------------------------------------------------
 
-void UReplicationGraphNode_ConnectionDormanyNode::GatherActorListsForConnection(const FConnectionGatherActorListParameters& Params)
+void UReplicationGraphNode_ConnectionDormancyNode::GatherActorListsForConnection(const FConnectionGatherActorListParameters& Params)
 {
 	ConditionalGatherDormantActorsForConnection(ReplicationActorList, Params, nullptr);
 	
@@ -3523,7 +3523,7 @@ void UReplicationGraphNode_ConnectionDormanyNode::GatherActorListsForConnection(
 			if (!RemoveList)
 			{
 				RemoveList = new (RemovedStreamingLevelActorListCollection.StreamingLevelLists) FStreamingLevelActorListCollection::FStreamingLevelActors(StreamingList.StreamingLevelName);
-				Params.ConnectionManager.OnClientVisibleLevelNameAddMap.FindOrAdd(StreamingList.StreamingLevelName).AddUObject(this, &UReplicationGraphNode_ConnectionDormanyNode::OnClientVisibleLevelNameAdd);
+				Params.ConnectionManager.OnClientVisibleLevelNameAddMap.FindOrAdd(StreamingList.StreamingLevelName).AddUObject(this, &UReplicationGraphNode_ConnectionDormancyNode::OnClientVisibleLevelNameAdd);
 			}
 
 			ConditionalGatherDormantActorsForConnection(StreamingList.ReplicationActorList, Params, &RemoveList->ReplicationActorList);
@@ -3535,7 +3535,7 @@ void UReplicationGraphNode_ConnectionDormanyNode::GatherActorListsForConnection(
 	}
 }
 
-void UReplicationGraphNode_ConnectionDormanyNode::ConditionalGatherDormantActorsForConnection(FActorRepListRefView& ConnectionList, const FConnectionGatherActorListParameters& Params, FActorRepListRefView* RemovedList)
+void UReplicationGraphNode_ConnectionDormancyNode::ConditionalGatherDormantActorsForConnection(FActorRepListRefView& ConnectionList, const FConnectionGatherActorListParameters& Params, FActorRepListRefView* RemovedList)
 {
 	FPerConnectionActorInfoMap& ConnectionActorInfoMap = Params.ConnectionManager.ActorInfoMap;
 	FGlobalActorReplicationInfoMap* GlobalActorReplicationInfoMap = GraphGlobals->GlobalActorReplicationInfoMap;
@@ -3567,7 +3567,7 @@ void UReplicationGraphNode_ConnectionDormanyNode::ConditionalGatherDormantActors
 			UE_CLOG(CVar_RepGraph_LogNetDormancyDetails > 0, LogReplicationGraph, Display, TEXT("GRAPH_DORMANCY: Actor %s is Dormant on %s. Removing from list. (%d elements left)"), *Actor->GetPathName(), *GetName(), ConnectionList.Num());
 			bShouldTrickle = false; // Dont trickle this frame because we are still encountering dormant actors
 		}
-		else if (CVar_RepGraph_TrickleDistCullOnDormanyNodes > 0 && bShouldTrickle)
+		else if (CVar_RepGraph_TrickleDistCullOnDormancyNodes > 0 && bShouldTrickle)
 		{
 			ConnectionActorInfo.SetCullDistanceSquared(0.f);
 			bShouldTrickle = false; // trickle one actor per frame
@@ -3596,7 +3596,7 @@ bool ContainsReverse(const FActorRepListRefView& List, FActorRepListType Actor)
 	return false;
 }
 
-void UReplicationGraphNode_ConnectionDormanyNode::NotifyActorDormancyFlush(FActorRepListType Actor)
+void UReplicationGraphNode_ConnectionDormancyNode::NotifyActorDormancyFlush(FActorRepListType Actor)
 {
 	FNewReplicatedActorInfo ActorInfo(Actor);
 
@@ -3634,7 +3634,7 @@ void UReplicationGraphNode_ConnectionDormanyNode::NotifyActorDormancyFlush(FActo
 	}
 }
 
-void UReplicationGraphNode_ConnectionDormanyNode::OnClientVisibleLevelNameAdd(FName LevelName, UWorld* World)
+void UReplicationGraphNode_ConnectionDormancyNode::OnClientVisibleLevelNameAdd(FName LevelName, UWorld* World)
 {
 	FStreamingLevelActorListCollection::FStreamingLevelActors* RemoveList = RemovedStreamingLevelActorListCollection.StreamingLevelLists.FindByKey(LevelName);
 	if (!RemoveList)
@@ -3660,7 +3660,7 @@ void UReplicationGraphNode_ConnectionDormanyNode::OnClientVisibleLevelNameAdd(FN
 	RemoveList->ReplicationActorList.Reset();
 }
 
-bool UReplicationGraphNode_ConnectionDormanyNode::NotifyRemoveNetworkActor(const FNewReplicatedActorInfo& ActorInfo, bool WarnIfNotFound)
+bool UReplicationGraphNode_ConnectionDormancyNode::NotifyRemoveNetworkActor(const FNewReplicatedActorInfo& ActorInfo, bool WarnIfNotFound)
 {
 	// Remove from active list by calling super
 	if (Super::NotifyRemoveNetworkActor(ActorInfo, false))
@@ -3672,7 +3672,7 @@ bool UReplicationGraphNode_ConnectionDormanyNode::NotifyRemoveNetworkActor(const
 	return RemovedStreamingLevelActorListCollection.RemoveActor(ActorInfo, WarnIfNotFound, this);
 }
 
-void UReplicationGraphNode_ConnectionDormanyNode::NotifyResetAllNetworkActors()
+void UReplicationGraphNode_ConnectionDormancyNode::NotifyResetAllNetworkActors()
 {
 	Super::NotifyResetAllNetworkActors();
 	RemovedStreamingLevelActorListCollection.Reset();
@@ -3715,7 +3715,7 @@ void UReplicationGraphNode_DormancyNode::AddDormantActor(const FNewReplicatedAct
 	
 	for (auto& MapIt : ConnectionNodes)
 	{
-		UReplicationGraphNode_ConnectionDormanyNode* Node = MapIt.Value;
+		UReplicationGraphNode_ConnectionDormancyNode* Node = MapIt.Value;
 		Node->NotifyAddNetworkActor(ActorInfo);
 	}
 
@@ -3734,7 +3734,7 @@ void UReplicationGraphNode_DormancyNode::RemoveDormantActor(const FNewReplicated
 	// Update any connection specific nodes
 	for (auto& MapIt : ConnectionNodes)
 	{
-		UReplicationGraphNode_ConnectionDormanyNode* Node = MapIt.Value;
+		UReplicationGraphNode_ConnectionDormancyNode* Node = MapIt.Value;
 		Node->NotifyRemoveNetworkActor(ActorInfo, false); // Don't warn if not found, the node may have removed the actor itself. Not worth the extra bookkeeping to skip the call.
 	}
 }
@@ -3756,24 +3756,24 @@ void UReplicationGraphNode_DormancyNode::GatherActorListsForConnection(const FCo
 		return;
 	}
 
-	UReplicationGraphNode_ConnectionDormanyNode* ConnectionNode = GetConnectionNode(Params);
+	UReplicationGraphNode_ConnectionDormancyNode* ConnectionNode = GetConnectionNode(Params);
 	ConnectionNode->GatherActorListsForConnection(Params);
 }
 
-UReplicationGraphNode_ConnectionDormanyNode* UReplicationGraphNode_DormancyNode::GetExistingConnectionNode(const FConnectionGatherActorListParameters& Params)
+UReplicationGraphNode_ConnectionDormancyNode* UReplicationGraphNode_DormancyNode::GetExistingConnectionNode(const FConnectionGatherActorListParameters& Params)
 {
-	UReplicationGraphNode_ConnectionDormanyNode** ConnectionNodeItem = ConnectionNodes.Find(&Params.ConnectionManager);
+	UReplicationGraphNode_ConnectionDormancyNode** ConnectionNodeItem = ConnectionNodes.Find(&Params.ConnectionManager);
 	return ConnectionNodeItem == nullptr ? nullptr : *ConnectionNodeItem;
 }
 
-UReplicationGraphNode_ConnectionDormanyNode* UReplicationGraphNode_DormancyNode::GetConnectionNode(const FConnectionGatherActorListParameters& Params)
+UReplicationGraphNode_ConnectionDormancyNode* UReplicationGraphNode_DormancyNode::GetConnectionNode(const FConnectionGatherActorListParameters& Params)
 {
-	UReplicationGraphNode_ConnectionDormanyNode** NodePtrPtr = ConnectionNodes.Find(&Params.ConnectionManager);
-	UReplicationGraphNode_ConnectionDormanyNode* ConnectionNode = nullptr;
+	UReplicationGraphNode_ConnectionDormancyNode** NodePtrPtr = ConnectionNodes.Find(&Params.ConnectionManager);
+	UReplicationGraphNode_ConnectionDormancyNode* ConnectionNode = nullptr;
 	if (!NodePtrPtr)
 	{
 		// We dont have a per-connection node for this connection, so create one and copy over contents
-		ConnectionNode = CreateChildNode<UReplicationGraphNode_ConnectionDormanyNode>();
+		ConnectionNode = CreateChildNode<UReplicationGraphNode_ConnectionDormancyNode>();
 		ConnectionNodes.Add(&Params.ConnectionManager) = ConnectionNode;
 
 		// Copy our master lists to the connection node
@@ -3815,7 +3815,7 @@ void UReplicationGraphNode_DormancyNode::OnActorDormancyFlush(FActorRepListType
 
 	for (auto& MapIt : ConnectionNodes)
 	{
-		UReplicationGraphNode_ConnectionDormanyNode* Node = MapIt.Value;
+		UReplicationGraphNode_ConnectionDormancyNode* Node = MapIt.Value;
 		Node->NotifyActorDormancyFlush(Actor);
 	}
 }
@@ -3933,7 +3933,7 @@ void UReplicationGraphNode_GridCell::OnStaticActorNetDormancyChange(FActorRepLis
 
 	if (!bCurrentDormant && bPreviousDormant)
 	{
-		// Actor is now awake, remove from dormany node and add to non dormany list
+		// Actor is now awake, remove from dormancy node and add to non dormancy list
 		FNewReplicatedActorInfo ActorInfo(Actor);
 		GetDormancyNode()->RemoveDormantActor(ActorInfo, GlobalInfo);
 		Super::NotifyAddNetworkActor(ActorInfo);
@@ -4943,7 +4943,7 @@ void UReplicationGraphNode_GridSpatialization2D::GatherActorListsForConnection(c
 						if (UReplicationGraphNode_DormancyNode* DormancyNode = Node->GetDormancyNode())
 						{
 							// Only notify the connection node if this client was previously inside the cell.
-							if (UReplicationGraphNode_ConnectionDormanyNode* ConnectionDormancyNode = DormancyNode->GetExistingConnectionNode(Params))
+							if (UReplicationGraphNode_ConnectionDormancyNode* ConnectionDormancyNode = DormancyNode->GetExistingConnectionNode(Params))
 							{
 								ConnectionDormancyNode->NotifyActorDormancyFlush(Actor);
 							}
diff --git a/Engine/Plugins/Runtime/ReplicationGraph/Source/Public/ReplicationGraph.h b/Engine/Plugins/Runtime/ReplicationGraph/Source/Public/ReplicationGraph.h
--- a/Engine/Plugins/Runtime/ReplicationGraph/Source/Public/ReplicationGraph.h
+++ b/Engine/Plugins/Runtime/ReplicationGraph/Source/Public/ReplicationGraph.h
@@ -376,7 +376,7 @@ protected:
 
 /** Removes dormant (on connection) actors from its rep lists */
 UCLASS()
-class REPLICATIONGRAPH_API UReplicationGraphNode_ConnectionDormanyNode : public UReplicationGraphNode_ActorList
+class REPLICATIONGRAPH_API UReplicationGraphNode_ConnectionDormancyNode : public UReplicationGraphNode_ActorList
 {
 	GENERATED_BODY()
 public:
@@ -421,13 +421,13 @@ public:
 
 	void ConditionalGatherDormantDynamicActors(FActorRepListRefView& RepList, const FConnectionGatherActorListParameters& Params, FActorRepListRefView* RemovedList, bool bEnforceReplistUniqueness=false);
 
-	UReplicationGraphNode_ConnectionDormanyNode* GetExistingConnectionNode(const FConnectionGatherActorListParameters& Params);
+	UReplicationGraphNode_ConnectionDormancyNode* GetExistingConnectionNode(const FConnectionGatherActorListParameters& Params);
 
-	UReplicationGraphNode_ConnectionDormanyNode* GetConnectionNode(const FConnectionGatherActorListParameters& Params);
+	UReplicationGraphNode_ConnectionDormancyNode* GetConnectionNode(const FConnectionGatherActorListParameters& Params);
 
 private:
 
-	TMap<UNetReplicationGraphConnection*, UReplicationGraphNode_ConnectionDormanyNode*> ConnectionNodes;
+	TMap<UNetReplicationGraphConnection*, UReplicationGraphNode_ConnectionDormancyNode*> ConnectionNodes;
 };
 
 UCLASS()
diff --git a/Engine/Plugins/Runtime/ReplicationGraph/Source/Public/ReplicationGraphTypes.h b/Engine/Plugins/Runtime/ReplicationGraph/Source/Public/ReplicationGraphTypes.h
--- a/Engine/Plugins/Runtime/ReplicationGraph/Source/Public/ReplicationGraphTypes.h
+++ b/Engine/Plugins/Runtime/ReplicationGraph/Source/Public/ReplicationGraphTypes.h
@@ -463,7 +463,7 @@ struct FGlobalActorReplicationInfo
 	/** Cached World Location of the actor */
 	FVector WorldLocation;
 
-	/** Mirrors AActor::NetDormany > DORM_Awake */
+	/** Mirrors AActor::NetDormancy > DORM_Awake */
 	bool bWantsToBeDormant = false;
 
 	/** Class default mirrors: state that is initialized directly from class defaults (and can be later changed on a per-actor basis) */
