Add ability to track in CSV the nb of connections who have a streaming level visible

#rnx

diff --git a/Engine/Plugins/Runtime/ReplicationGraph/Source/Private/ReplicationGraph.cpp b/Engine/Plugins/Runtime/ReplicationGraph/Source/Private/ReplicationGraph.cpp
--- a/Engine/Plugins/Runtime/ReplicationGraph/Source/Private/ReplicationGraph.cpp
+++ b/Engine/Plugins/Runtime/ReplicationGraph/Source/Private/ReplicationGraph.cpp
@@ -116,6 +116,7 @@ CSV_DEFINE_CATEGORY(ReplicationGraphMS, WITH_SERVER_CODE);
 CSV_DEFINE_CATEGORY(ReplicationGraphKBytes, WITH_SERVER_CODE);
 CSV_DEFINE_CATEGORY(ReplicationGraphChannelsOpened, WITH_SERVER_CODE);
 CSV_DEFINE_CATEGORY(ReplicationGraphNumReps, WITH_SERVER_CODE);
+CSV_DEFINE_CATEGORY(ReplicationGraphVisibleLevels, WITH_SERVER_CODE);
 
 // Helper CVar for debugging. Set this string to conditionally log/breakpoint various points in the repgraph pipeline. Useful for bugs like "why is this actor channel closing"
 static TAutoConsoleVariable<FString> CVarRepGraphConditionalBreakpointActorName(TEXT("Net.RepGraph.ConditionalBreakpointActorName"), TEXT(""), TEXT(""), ECVF_Default );
diff --git a/Engine/Plugins/Runtime/ReplicationGraph/Source/Public/ReplicationGraphTypes.h b/Engine/Plugins/Runtime/ReplicationGraph/Source/Public/ReplicationGraphTypes.h
--- a/Engine/Plugins/Runtime/ReplicationGraph/Source/Public/ReplicationGraphTypes.h
+++ b/Engine/Plugins/Runtime/ReplicationGraph/Source/Public/ReplicationGraphTypes.h
@@ -1319,6 +1319,7 @@ CSV_DECLARE_CATEGORY_EXTERN(ReplicationGraphMS);
 CSV_DECLARE_CATEGORY_EXTERN(ReplicationGraphKBytes);
 CSV_DECLARE_CATEGORY_EXTERN(ReplicationGraphChannelsOpened);
 CSV_DECLARE_CATEGORY_EXTERN(ReplicationGraphNumReps);
+CSV_DECLARE_CATEGORY_EXTERN(ReplicationGraphVisibleLevels);
 
 #ifndef REPGRAPH_CSV_TRACKER
 #define REPGRAPH_CSV_TRACKER (CSV_PROFILER && WITH_SERVER_CODE)
@@ -1336,26 +1337,81 @@ struct FReplicationGraphCSVTracker
 	}
 
 	/** Tracks an explicitly set class. This does NOT include child classes! This is the fastest stat and should be fine to enable in shipping/test builds.  */
-	void SetExplicitClassTracking(UClass* ExactActorClass, FString StatNamePrefix)
+	void SetExplicitClassTracking(UClass* ExactActorClass, const FString& StatNamePrefix)
 	{
 		ExplicitClassTracker.Emplace(ExactActorClass, StatNamePrefix);
 	}
 
 	/** Sets explicit class tracking for fast/shared path replication. Does not include base classes */
-	void SetExplicitClassTracking_FastPath(UClass* ExactActorClass, FString StatNamePrefix)
+	void SetExplicitClassTracking_FastPath(UClass* ExactActorClass, const FString& StatNamePrefix)
 	{
 		FString FinalStrPrefix = TEXT("FastPath_") + StatNamePrefix;
 		ExplicitClassTracker_FastPath.Emplace(ExactActorClass, FinalStrPrefix);
 	}
 
 	/** Tracks a class and all of its children (under a single stat set). This will be a little slower (TMap lookup) but still probably ok if used in moderation (only track your top 3 or so classes) */
-	void SetImplicitClassTracking(UClass* BaseActorClass, FString StatNamePrefix)
+	void SetImplicitClassTracking(UClass* BaseActorClass, const FString& StatNamePrefix)
 	{
 		TSharedPtr<FTrackedData> NewData(new FTrackedData(StatNamePrefix));
 		UniqueImplicitTrackedData.Add(NewData);
 		ImplicitClassTracker.Set(BaseActorClass, NewData);
 	}
 
+	void VisibleLevelConnectionAdded(FName LevelName)
+	{
+#if REPGRAPH_CSV_TRACKER
+		FVisibleLevelData* LevelData = VisibleLevelConnectionTracker.FindByKey(LevelName);
+		if (LevelData)
+		{
+			LevelData->NbConnections++;
+		}
+		else
+		{
+			FVisibleLevelData VisibleLevelData;
+			VisibleLevelData.LevelName = LevelName;
+			VisibleLevelData.NbConnections = 1;
+			VisibleLevelConnectionTracker.Add(VisibleLevelData);
+		}
+#endif
+	}
+
+	void VisibleLevelConnectionRemoved(FName LevelName)
+	{
+#if REPGRAPH_CSV_TRACKER
+		if (FVisibleLevelData* LevelData = VisibleLevelConnectionTracker.FindByKey(LevelName))
+		{
+			LevelData->NbConnections--;
+			if (LevelData->NbConnections <= 0)
+			{
+				VisibleLevelConnectionTracker.RemoveSingleSwap(*LevelData);
+			}
+		}
+#endif
+	}
+
+	void SetVisibleLevelCustomName(FName LevelName, const FString& CustomReadableName)
+	{
+#if REPGRAPH_CSV_TRACKER
+		const FName NameReadable(*CustomReadableName);
+		FVisibleLevelData* LevelData = VisibleLevelConnectionTracker.FindByKey(LevelName);
+		if (LevelData)
+		{
+			// Note this will create a new stat entry in the csv file
+			LevelData->CustomReadableName = NameReadable;
+		}
+		else
+		{
+			FVisibleLevelData VisibleLevelData;
+			VisibleLevelData.LevelName = LevelName;
+			VisibleLevelData.NbConnections = 0;
+			VisibleLevelData.CustomReadableName = NameReadable;
+
+			VisibleLevelConnectionTracker.Add(VisibleLevelData);
+		}
+#endif //REPGRAPH_CSV_TRACKER
+	}
+
+
 	void PostReplicateActor(UClass* ActorClass, const double Time, const int64 Bits, const bool bIsActorDiscovery)
 	{
 #if REPGRAPH_CSV_TRACKER
@@ -1451,6 +1507,7 @@ struct FReplicationGraphCSVTracker
 		EverythingElse.Reset();
 		EverythingElse_FastPath.Reset();
 		ActorDiscovery.Reset();
+		VisibleLevelConnectionTracker.Reset();
 	}
 
 	void EndReplicationFrame()
@@ -1481,6 +1538,18 @@ struct FReplicationGraphCSVTracker
 			PushStats(Profiler, EverythingElse);
 			PushStats(Profiler, EverythingElse_FastPath);
 			PushStats(Profiler, ActorDiscovery);
+
+			for (const FVisibleLevelData& VisibleLevel : VisibleLevelConnectionTracker)
+			{
+				if (VisibleLevel.CustomReadableName.IsValid())
+				{
+					Profiler->RecordCustomStat(VisibleLevel.CustomReadableName, CSV_CATEGORY_INDEX(ReplicationGraphVisibleLevels), VisibleLevel.NbConnections, ECsvCustomStatOp::Set);
+				}
+				else
+				{
+					Profiler->RecordCustomStat(VisibleLevel.LevelName, CSV_CATEGORY_INDEX(ReplicationGraphVisibleLevels), VisibleLevel.NbConnections, ECsvCustomStatOp::Set);
+				}
+			}
 		}
 #endif
 	}
@@ -1528,6 +1597,16 @@ private:
 		FTrackedData Data;
 	};
 
+	struct FVisibleLevelData
+	{
+		FName LevelName;
+		int32 NbConnections;
+		FName CustomReadableName;
+
+		bool operator==(const FVisibleLevelData& Other) const { return LevelName == Other.LevelName; }
+		bool operator==(const FName& InLevelName) const { return LevelName == InLevelName; }
+	};
+
 	TArray<FTrackerItem, TInlineAllocator<1>> ExplicitClassTracker;
 	TClassMap<TSharedPtr<FTrackedData>> ImplicitClassTracker;
 	TArray<TSharedPtr<FTrackedData>> UniqueImplicitTrackedData;
@@ -1538,6 +1617,9 @@ private:
 
 	FTrackedData ActorDiscovery;
 
+	// Counts the number of connections who are currently seeing each level name
+	TArray<FVisibleLevelData> VisibleLevelConnectionTracker;
+
 	bool bIsCapturing = false;
 
 #if REPGRAPH_CSV_TRACKER
