RepGraph
- Ignore Dormancy changes for actors not added to the RepGraph yet.  Their dormancy status will be checked when added as a network actor.

#rnx

diff --git a/Engine/Plugins/Runtime/ReplicationGraph/Source/Private/ReplicationGraph.cpp b/Engine/Plugins/Runtime/ReplicationGraph/Source/Private/ReplicationGraph.cpp
--- a/Engine/Plugins/Runtime/ReplicationGraph/Source/Private/ReplicationGraph.cpp
+++ b/Engine/Plugins/Runtime/ReplicationGraph/Source/Private/ReplicationGraph.cpp
@@ -694,16 +694,22 @@ void UReplicationGraph::NotifyActorDormancyChange(AActor* Actor, ENetDormancy Ol
 {
 	RG_QUICK_SCOPE_CYCLE_COUNTER(UReplicationGraph_NotifyActorDormancyChange);
 
-	FGlobalActorReplicationInfo& GlobalInfo = GlobalActorReplicationInfoMap.Get(Actor);
+	FGlobalActorReplicationInfo* GlobalInfo = GlobalActorReplicationInfoMap.Find(Actor);
+	if (!GlobalInfo)
+	{
+		UE_CLOG(CVar_RepGraph_LogNetDormancyDetails > 0, LogReplicationGraph, Display, TEXT("UReplicationGraph::NotifyActorDormancyChange %s. Ignoring change since actor is not registered yet."), *Actor->GetPathName());
+		return;
+	}
+
 	ENetDormancy CurrentDormancy = Actor->NetDormancy;
 
-	UE_CLOG(CVar_RepGraph_LogNetDormancyDetails > 0, LogReplicationGraph, Display, TEXT("UReplicationGraph::NotifyActorDormancyChange %s. Old WantsToBeDormant: %d. New WantsToBeDormant: %d"), *Actor->GetPathName(), GlobalInfo.bWantsToBeDormant, CurrentDormancy > DORM_Awake ? 1 : 0);
+	UE_CLOG(CVar_RepGraph_LogNetDormancyDetails > 0, LogReplicationGraph, Display, TEXT("UReplicationGraph::NotifyActorDormancyChange %s. Old WantsToBeDormant: %d. New WantsToBeDormant: %d"), *Actor->GetPathName(), GlobalInfo->bWantsToBeDormant, CurrentDormancy > DORM_Awake ? 1 : 0);
 
 	const bool bOldWantsToBeDormant = OldDormancyState > DORM_Awake;
 	const bool bNewWantsToBeDormant = CurrentDormancy > DORM_Awake;
 
-	GlobalInfo.bWantsToBeDormant = bNewWantsToBeDormant;
-	GlobalInfo.Events.DormancyChange.Broadcast(Actor, GlobalInfo, CurrentDormancy, OldDormancyState);
+	GlobalInfo->bWantsToBeDormant = bNewWantsToBeDormant;
+	GlobalInfo->Events.DormancyChange.Broadcast(Actor, *GlobalInfo, CurrentDormancy, OldDormancyState);
 
 	// Is the actor coming out of dormancy via changing its dormancy state?
 	if (!bNewWantsToBeDormant && bOldWantsToBeDormant)
