Improve a log message in UReplicationGraph::ProcessRemoteFunction when an actor is not valid for replication or is being destroyed.

diff --git a/Engine/Plugins/Runtime/ReplicationGraph/Source/Private/ReplicationGraph.cpp b/Engine/Plugins/Runtime/ReplicationGraph/Source/Private/ReplicationGraph.cpp
--- a/Engine/Plugins/Runtime/ReplicationGraph/Source/Private/ReplicationGraph.cpp
+++ b/Engine/Plugins/Runtime/ReplicationGraph/Source/Private/ReplicationGraph.cpp
@@ -1832,12 +1832,12 @@ bool UReplicationGraph::ProcessRemoteFunction(class AActor* Actor, UFunction* Fu
 
 	if (RepGraphConditionalActorBreakpoint(Actor, nullptr))
 	{
-		UE_LOG(LogReplicationGraph, Display, TEXT("UReplicationGraph::ProcessRemoteFunction: %s. Function: %s."), *Actor->GetName(), *GetNameSafe(Function));
+		UE_LOG(LogReplicationGraph, Display, TEXT("UReplicationGraph::ProcessRemoteFunction: %s. Function: %s."), *GetNameSafe(Actor), *GetNameSafe(Function));
 	}
 
 	if (IsActorValidForReplication(Actor) == false || Actor->IsActorBeingDestroyed())
 	{
-		UE_LOG(LogReplicationGraph, Display, TEXT("Destroted or not ready!"));
+		UE_LOG(LogReplicationGraph, Display, TEXT("UReplicationGraph::ProcessRemoteFunction: Actor %s destroyed or not ready! Function: %s."), *GetNameSafe(Actor), *GetNameSafe(Function));
 		return true;
 	}
 
