Ignore FlushNetDormancy and DormancyChange for destroyed or about to be destroyed actors

#rnx

diff --git a/Engine/Plugins/Runtime/ReplicationGraph/Source/Private/ReplicationGraph.cpp b/Engine/Plugins/Runtime/ReplicationGraph/Source/Private/ReplicationGraph.cpp
--- a/Engine/Plugins/Runtime/ReplicationGraph/Source/Private/ReplicationGraph.cpp
+++ b/Engine/Plugins/Runtime/ReplicationGraph/Source/Private/ReplicationGraph.cpp
@@ -597,6 +597,12 @@ void UReplicationGraph::FlushNetDormancy(AActor* Actor, bool bWasDormInitial)
 		return;
 	}
 
+	if (IsActorValidForReplication(Actor) == false)
+	{
+		UE_CLOG(CVar_RepGraph_LogNetDormancyDetails > 0, LogReplicationGraph, Display, TEXT("UReplicationGraph::FlushNetDormancy called on %s. Ignored since actor is destroyed or about to be"), *Actor->GetPathName());
+		return;
+	}
+
 	FGlobalActorReplicationInfo& GlobalInfo = GlobalActorReplicationInfoMap.Get(Actor);
 	const bool bNewWantsToBeDormant = (Actor->NetDormancy > DORM_Awake);
 
@@ -701,6 +707,12 @@ void UReplicationGraph::NotifyActorDormancyChange(AActor* Actor, ENetDormancy Ol
 		return;
 	}
 
+	if (IsActorValidForReplication(Actor) == false)
+	{
+		UE_CLOG(CVar_RepGraph_LogNetDormancyDetails > 0, LogReplicationGraph, Display, TEXT("UReplicationGraph::NotifyActorDormancyChange %s. Ignoring change since actor is destroyed or about to be."), *Actor->GetPathName());
+		return;
+	}
+
 	ENetDormancy CurrentDormancy = Actor->NetDormancy;
 
 	UE_CLOG(CVar_RepGraph_LogNetDormancyDetails > 0, LogReplicationGraph, Display, TEXT("UReplicationGraph::NotifyActorDormancyChange %s. Old WantsToBeDormant: %d. New WantsToBeDormant: %d"), *Actor->GetPathName(), GlobalInfo->bWantsToBeDormant, CurrentDormancy > DORM_Awake ? 1 : 0);
@@ -5244,4 +5256,4 @@ FAutoConsoleCommandWithWorldAndArgs ChangeActorDiscoveryBudget(TEXT("Net.RepGrap
 	{
 		It->SetActorDiscoveryBudget(BudgetInKBPS);
 	}
-}));
\ No newline at end of file
+}));
