PR #5540: RepGraph: Fix stale actor channels not cleaning up (Contributed by [REDACTED_PERSON])

diff --git a/Engine/Plugins/Runtime/ReplicationGraph/Source/Private/ReplicationGraph.cpp b/Engine/Plugins/Runtime/ReplicationGraph/Source/Private/ReplicationGraph.cpp
--- a/Engine/Plugins/Runtime/ReplicationGraph/Source/Private/ReplicationGraph.cpp
+++ b/Engine/Plugins/Runtime/ReplicationGraph/Source/Private/ReplicationGraph.cpp
@@ -2066,13 +2066,14 @@ bool UReplicationGraph::ProcessRemoteFunction(class AActor* Actor, UFunction* Fu
 		UActorChannel* Ch = Connection->FindActorChannelRef(Actor);
 		if (Ch == nullptr)
 		{
-			if ( Actor->IsPendingKillPending() || !NetDriver->IsLevelInitializedForActor(Actor, Connection) )
+			if (Actor->IsPendingKillPending() || !NetDriver->IsLevelInitializedForActor(Actor, Connection))
 			{
 				// We can't open a channel for this actor here
 				return true;
 			}
 
-			if (UNetReplicationGraphConnection* ConnectionManager = Cast<UNetReplicationGraphConnection>(Connection->GetReplicationConnectionDriver()))
+			UNetReplicationGraphConnection* ConnectionManager = Cast<UNetReplicationGraphConnection>(Connection->GetReplicationConnectionDriver());
+			if (ConnectionManager)
 			{
 				if (FConnectionReplicationActorInfo const * const ConnectionActorInfo = ConnectionManager->ActorInfoMap.Find(Actor))
 				{
@@ -2088,6 +2089,13 @@ bool UReplicationGraph::ProcessRemoteFunction(class AActor* Actor, UFunction* Fu
 
 			Ch = (UActorChannel *)Connection->CreateChannelByName( NAME_Actor, EChannelCreateFlags::OpenedLocally );
 			Ch->SetChannelActor(Actor, ESetChannelActorFlags::None);
+			
+			if (ConnectionManager)
+			{
+				FConnectionReplicationActorInfo& ConnectionActorInfo = ConnectionManager->ActorInfoMap.FindOrAdd(Actor);
+				FGlobalActorReplicationInfo& GlobalInfo = GlobalActorReplicationInfoMap.Get(Actor);
+				UpdateActorChannelCloseFrameNum(Actor, ConnectionActorInfo, GlobalInfo, ReplicationGraphFrame+1 /** Plus one to error on safe side. RepFrame num will be incremented in the next tick */, Connection );
+			}
 		}
 
 		NetDriver->ProcessRemoteFunctionForChannel(Ch, ClassCache, FieldCache, TargetObj, Connection, Function, Parameters, OutParms, Stack, true);
