Fix missing files.

#rnx

diff --git a/Engine/Plugins/Runtime/ReplicationGraph/Source/Public/ReplicationGraph.h b/Engine/Plugins/Runtime/ReplicationGraph/Source/Public/ReplicationGraph.h
--- a/Engine/Plugins/Runtime/ReplicationGraph/Source/Public/ReplicationGraph.h
+++ b/Engine/Plugins/Runtime/ReplicationGraph/Source/Public/ReplicationGraph.h
@@ -889,6 +889,8 @@ public:
 
 	void UpdateActorChannelCloseFrameNum(AActor* Actor, FConnectionReplicationActorInfo& ConnectionData, const FGlobalActorReplicationInfo& GlobalData, const uint32 FrameNum, UNetConnection* NetConnection) const;
 
+	void NotifyConnectionSaturated(class UNetReplicationGraphConnection& Connection);
+
 protected:
 
 	virtual void InitializeForWorld(UWorld* World);
@@ -956,6 +958,9 @@ protected:
 
 private:
 
+	/** Whether or not a connection was saturated during an update. */
+	bool bWasConnectionSaturated = false;
+
 	/** Internal frame counter for replication. This is only updated by us. The one of UNetDriver can be updated by RPC calls and is only used to invalidate shared property CLs/serialiation data. */
 	uint32 ReplicationGraphFrame = 0;
 
diff --git a/Engine/Plugins/Runtime/ReplicationGraph/Source/ReplicationGraph.Build.cs b/Engine/Plugins/Runtime/ReplicationGraph/Source/ReplicationGraph.Build.cs
--- a/Engine/Plugins/Runtime/ReplicationGraph/Source/ReplicationGraph.Build.cs
+++ b/Engine/Plugins/Runtime/ReplicationGraph/Source/ReplicationGraph.Build.cs
@@ -13,7 +13,8 @@ public class ReplicationGraph : ModuleRules
 				"Core",
 				"CoreUObject",
 				"Engine",
-				"EngineSettings"
+				"EngineSettings",
+				"PerfCounters"
 			}
 		);
 
