Modify some sanity checks in RepGraph to allow RPCs for Actors that previously had channels that were closed (but not cleaned up).
Add a new ensure in the case where we try to reassign an Actor's channel when the existing channel is still valid.

diff --git a/Engine/Plugins/Runtime/ReplicationGraph/Source/Private/ReplicationGraph.cpp b/Engine/Plugins/Runtime/ReplicationGraph/Source/Private/ReplicationGraph.cpp
--- a/Engine/Plugins/Runtime/ReplicationGraph/Source/Private/ReplicationGraph.cpp
+++ b/Engine/Plugins/Runtime/ReplicationGraph/Source/Private/ReplicationGraph.cpp
@@ -2084,25 +2084,10 @@ bool UReplicationGraph::ProcessRemoteFunction(class AActor* Actor, UFunction* Fu
 				return true;
 			}
 
-			UNetReplicationGraphConnection* ConnectionManager = Cast<UNetReplicationGraphConnection>(Connection->GetReplicationConnectionDriver());
-			if (ConnectionManager)
-			{
-				if (FConnectionReplicationActorInfo const * const ConnectionActorInfo = ConnectionManager->ActorInfoMap.Find(Actor))
-				{
-					const bool bIsValid = ensureMsgf(ConnectionActorInfo->Channel == nullptr, TEXT("UReplicationGraph::ProcessRemoteFunction: Trying to send RPC for Actor whose channel is in invalid state %s"), *ConnectionActorInfo->Channel->Describe());
-
-					// Don't try to open a channel for an unreliable RPC whose channel is closing.
-					if (!(bIsReliable || bIsValid))
-					{
-						return true;
-					}
-				}
-			}
-
 			Ch = (UActorChannel *)Connection->CreateChannelByName( NAME_Actor, EChannelCreateFlags::OpenedLocally );
 			Ch->SetChannelActor(Actor, ESetChannelActorFlags::None);
 			
-			if (ConnectionManager)
+			if (UNetReplicationGraphConnection* ConnectionManager = Cast<UNetReplicationGraphConnection>(Connection->GetReplicationConnectionDriver()))
 			{
 				FConnectionReplicationActorInfo& ConnectionActorInfo = ConnectionManager->ActorInfoMap.FindOrAdd(Actor);
 				FGlobalActorReplicationInfo& GlobalInfo = GlobalActorReplicationInfoMap.Get(Actor);
@@ -2235,6 +2220,9 @@ void UNetReplicationGraphConnection::NotifyActorChannelAdded(AActor* Actor, clas
 	if (ActorInfo.Channel && Channel != ActorInfo.Channel)
 	{
 		UE_LOG(LogReplicationGraph, Log, TEXT("::NotifyActorChannelAdded. Fixing up stale channel reference Old: %s New: %s"), *ActorInfo.Channel->Describe(), *Channel->Describe());
+		ensureMsgf(ActorInfo.Channel->Closing, TEXT("Attempted to add an Actor Channel when a valid channel already exists for the actor. Actor=%s, OldChannel=%s, NewChannel=%s"),
+			*GetPathNameSafe(Actor), *ActorInfo.Channel->Describe(), *Channel->Describe());
+
 		ActorInfoMap.RemoveChannel(ActorInfo.Channel);
 	}
 
