Rep graph GC fixes

diff --git a/Engine/Plugins/Runtime/ReplicationGraph/Source/Private/ReplicationGraph.cpp b/Engine/Plugins/Runtime/ReplicationGraph/Source/Private/ReplicationGraph.cpp
--- a/Engine/Plugins/Runtime/ReplicationGraph/Source/Private/ReplicationGraph.cpp
+++ b/Engine/Plugins/Runtime/ReplicationGraph/Source/Private/ReplicationGraph.cpp
@@ -4827,7 +4827,7 @@ struct FPlayerGridCellInformation
 
 void UReplicationGraphNode_GridSpatialization2D::GatherActorListsForConnection(const FConnectionGatherActorListParameters& Params)
 {
-	TArray<FLastLocationGatherInfo, FReplicationGraphConnectionsAllocator>& LastLocationArray = Params.ConnectionManager.LastGatherLocations;
+	TArray<FLastLocationGatherInfo>& LastLocationArray = Params.ConnectionManager.LastGatherLocations;
 	TArray<FVector2D, FReplicationGraphConnectionsAllocator> UniqueCurrentLocations;
 
 	// Consider all users that are in cells for this connection. 
diff --git a/Engine/Plugins/Runtime/ReplicationGraph/Source/Public/ReplicationGraph.h b/Engine/Plugins/Runtime/ReplicationGraph/Source/Public/ReplicationGraph.h
--- a/Engine/Plugins/Runtime/ReplicationGraph/Source/Public/ReplicationGraph.h
+++ b/Engine/Plugins/Runtime/ReplicationGraph/Source/Public/ReplicationGraph.h
@@ -709,7 +709,8 @@ public:
 	FActorRepListRefView ReplicationActorList;
 
 	/** List of previously (or currently if nothing changed last tick) focused actor data per connection */
-	TArray<FAlwaysRelevantActorInfo, FReplicationGraphConnectionsAllocator> PastRelevantActors;
+	UPROPERTY()
+	TArray<FAlwaysRelevantActorInfo> PastRelevantActors;
 
 	UPROPERTY()
 	AActor* LastViewer = nullptr;
@@ -1040,7 +1041,8 @@ public:
 
 	FVector LastGatherLocation;
 
-	TArray<FLastLocationGatherInfo, FReplicationGraphConnectionsAllocator> LastGatherLocations;
+	UPROPERTY()
+	TArray<FLastLocationGatherInfo> LastGatherLocations;
 
 	// Nb of bits sent for actor channel creation when a dedicated budget is allocated for this
 	int32 QueuedBitsForActorDiscovery = 0;
