More logs and another tentative fix for Server Ensure in UReplicationGraph::ServerReplicateActors
* DependantActors are now checked to prevent being added twice in the dependant list.
* Added a specific ensure for DependantActors so it triggers before calling ReplicateActor on them.

#rnx

diff --git a/Engine/Plugins/Runtime/ReplicationGraph/Source/Private/ReplicationGraph.cpp b/Engine/Plugins/Runtime/ReplicationGraph/Source/Private/ReplicationGraph.cpp
--- a/Engine/Plugins/Runtime/ReplicationGraph/Source/Private/ReplicationGraph.cpp
+++ b/Engine/Plugins/Runtime/ReplicationGraph/Source/Private/ReplicationGraph.cpp
@@ -1667,8 +1667,8 @@ int64 UReplicationGraph::ReplicateSingleActor(AActor* Actor, FConnectionReplicat
 		UE_LOG(LogReplicationGraph, Display, TEXT("UReplicationGraph::ReplicateSingleActor: %s. NetConnection: %s"), *Actor->GetName(), *NetConnection->Describe());
 	}
 
-	if (!ensureMsgf(IsActorValidForReplication(Actor), TEXT("Actor not valid for replication (BeingDestroyed:%d) (PendingKill:%d) (Unreachable:%d)! Actor = %s, Channel = %s"),
-					Actor->IsActorBeingDestroyed(), Actor->IsPendingKill(), Actor->IsUnreachable(),
+	if (!ensureMsgf(IsActorValidForReplication(Actor), TEXT("Actor not valid for replication (BeingDestroyed:%d) (PendingKill:%d) (Unreachable:%d) (TearOff:%d)! Actor = %s, Channel = %s"),
+					Actor->IsActorBeingDestroyed(), Actor->IsPendingKill(), Actor->IsUnreachable(), Actor->GetTearOff(),
 					*Actor->GetFullName(), *DescribeSafe(ActorInfo.Channel)))
 	{
 		return 0;
@@ -1791,6 +1791,14 @@ int64 UReplicationGraph::ReplicateSingleActor(AActor* Actor, FConnectionReplicat
 				continue;
 			}
 
+			if (!ensureMsgf(IsActorValidForReplication(DependentActor), TEXT("DependentActor %s (Owner: %s) not valid for replication (BeingDestroyed:%d) (PendingKill:%d) (Unreachable:%d) (TearOff:%d)! Channel = %s"),
+							*DependentActor->GetFullName(), *Actor->GetFullName(),
+							DependentActor->IsActorBeingDestroyed(), DependentActor->IsPendingKill(), DependentActor->IsUnreachable(), DependentActor->GetTearOff(),
+							*DescribeSafe(DependentActorConnectionInfo.Channel)))
+			{
+				continue;
+			}
+
 			//UE_LOG(LogReplicationGraph, Display, TEXT("DependentActor %s %s. NextReplicationFrameNum: %d. FrameNum: %d. ForceNetUpdateFrame: %d. LastRepFrameNum: %d."), *DependentActor->GetPathName(), *NetConnection->GetName(), DependentActorConnectionInfo.NextReplicationFrameNum, FrameNum, DependentActorGlobalData.ForceNetUpdateFrame, DependentActorConnectionInfo.LastRepFrameNum);
 			BitsWritten += ReplicateSingleActor(DependentActor, DependentActorConnectionInfo, DependentActorGlobalData, ConnectionActorInfoMap, ConnectionManager, FrameNum);
 		}					
diff --git a/Engine/Plugins/Runtime/ReplicationGraph/Source/Public/ReplicationGraphTypes.h b/Engine/Plugins/Runtime/ReplicationGraph/Source/Public/ReplicationGraphTypes.h
--- a/Engine/Plugins/Runtime/ReplicationGraph/Source/Public/ReplicationGraphTypes.h
+++ b/Engine/Plugins/Runtime/ReplicationGraph/Source/Public/ReplicationGraphTypes.h
@@ -736,7 +736,13 @@ struct FGlobalActorReplicationInfoMap
 		ClassMap.CountBytes(Ar);
 	}
 
-	void AddDependentActor(AActor* Parent, AActor* Child)
+	enum EWarnFlag
+	{
+		None = 0,
+		WarnAlreadyDependant = 1 << 1,
+	};
+
+	void AddDependentActor(AActor* Parent, AActor* Child, FGlobalActorReplicationInfoMap::EWarnFlag WarnFlag=FGlobalActorReplicationInfoMap::None)
 	{
 		const bool bIsParentValid = ensureMsgf(Parent && IsActorValidForReplication(Parent), TEXT("FGlobalActorReplicationInfoMap::AddDependentActor Invalid Parent! %s"),
 			*GetPathNameSafe(Parent));
@@ -746,17 +752,31 @@ struct FGlobalActorReplicationInfoMap
 
 		if (bIsParentValid && bIsChildValid)
 		{
+			const bool bDoWarnings = (WarnFlag & WarnAlreadyDependant) != 0;
+
+			bool bChildIsAlreadyDependant(false);
 			if (FGlobalActorReplicationInfo* ParentInfo = Find(Parent))
 			{
+				bChildIsAlreadyDependant = bDoWarnings && ParentInfo->DependentActorList.IsValid() && ParentInfo->DependentActorList.Contains(Child);
+
 				ParentInfo->DependentActorList.PrepareForWrite();
 				ParentInfo->DependentActorList.ConditionalAdd(Child);
 			}
 
+			bool bChildHadParentAlready(false);
 			if (FGlobalActorReplicationInfo* ChildInfo = Find(Child))
 			{
+				bChildHadParentAlready = bDoWarnings && ChildInfo->ParentActorList.IsValid() && ChildInfo->ParentActorList.Contains(Parent);
 				ChildInfo->ParentActorList.PrepareForWrite();
 				ChildInfo->ParentActorList.ConditionalAdd(Parent);
 			}
+
+			if (bChildIsAlreadyDependant || bChildHadParentAlready)
+			{
+				UE_LOG(LogReplicationGraph, Warning, TEXT("FGlobalActorReplicationInfoMap::AddDependentActor Child %s - Parent %s | Child already dependant of parent: %d | Child previously had parent in list: %d"), 
+					   *GetPathNameSafe(Child), *GetPathNameSafe(Parent),
+					   bChildIsAlreadyDependant, bChildHadParentAlready);
+			}
 		}
 	}
 
@@ -778,6 +798,43 @@ struct FGlobalActorReplicationInfoMap
 		}
 	}
 
+	void RemoveAllActorDependancies(AActor* MainActor)
+	{
+		FGlobalActorReplicationInfo* MainActorInfo = Find(MainActor);
+		if (MainActor == nullptr)
+		{
+			return;
+		}
+		
+		if (MainActorInfo->ParentActorList.IsValid())
+		{
+			// Remove this actor from all his parents
+			for (FActorRepListType ParentActor : MainActorInfo->ParentActorList)
+			{
+				if (FGlobalActorReplicationInfo* ParentInfo = Find(ParentActor))
+				{
+					ParentInfo->DependentActorList.PrepareForWrite();
+					ParentInfo->DependentActorList.Remove(MainActor);
+				}
+			}
+			MainActorInfo->ParentActorList.Reset();
+		}
+
+		if (MainActorInfo->DependentActorList.IsValid())
+		{
+            // Remove all dependant childs from this actor
+			for (FActorRepListType ChildActor : MainActorInfo->DependentActorList)
+			{
+				if (FGlobalActorReplicationInfo* ChildInfo = Find(ChildActor))
+				{
+					ChildInfo->ParentActorList.PrepareForWrite();
+					ChildInfo->ParentActorList.Remove(MainActor);
+				}
+			}
+			MainActorInfo->DependentActorList.Reset();
+		}
+	}
+
 private:
 
 	TMap<FActorRepListType, TUniquePtr<FGlobalActorReplicationInfo>> ActorMap;
