Fix error in dll builds with LogReplicationGraph category

#rnx

diff --git a/Engine/Plugins/Runtime/ReplicationGraph/Source/Private/ReplicationGraphTypes.cpp b/Engine/Plugins/Runtime/ReplicationGraph/Source/Private/ReplicationGraphTypes.cpp
--- a/Engine/Plugins/Runtime/ReplicationGraph/Source/Private/ReplicationGraphTypes.cpp
+++ b/Engine/Plugins/Runtime/ReplicationGraph/Source/Private/ReplicationGraphTypes.cpp
@@ -459,3 +459,41 @@ void TActorListAllocator<NumListsPerBlock, MaxNumPools>::LogDetails(int32 PoolSi
 	}
 }
 
+void FGlobalActorReplicationInfoMap::AddDependentActor(AActor* Parent, AActor* Child, FGlobalActorReplicationInfoMap::EWarnFlag WarnFlag)
+{
+	const bool bIsParentValid = ensureMsgf(Parent && IsActorValidForReplication(Parent), TEXT("FGlobalActorReplicationInfoMap::AddDependentActor Invalid Parent! %s"),
+										   *GetPathNameSafe(Parent));
+
+	const bool bIsChildValid = ensureMsgf(Child && IsActorValidForReplication(Child), TEXT("FGlobalActorReplicationInfoMap::AddDependentActor Invalid Child! %s"),
+										  *GetPathNameSafe(Child));
+
+	if (bIsParentValid && bIsChildValid)
+	{
+		const bool bDoWarnings = (WarnFlag & WarnAlreadyDependant) != 0;
+
+		bool bChildIsAlreadyDependant(false);
+		if (FGlobalActorReplicationInfo* ParentInfo = Find(Parent))
+		{
+			bChildIsAlreadyDependant = bDoWarnings && ParentInfo->DependentActorList.IsValid() && ParentInfo->DependentActorList.Contains(Child);
+
+			ParentInfo->DependentActorList.PrepareForWrite();
+			ParentInfo->DependentActorList.ConditionalAdd(Child);
+		}
+
+		bool bChildHadParentAlready(false);
+		if (FGlobalActorReplicationInfo* ChildInfo = Find(Child))
+		{
+			bChildHadParentAlready = bDoWarnings && ChildInfo->ParentActorList.IsValid() && ChildInfo->ParentActorList.Contains(Parent);
+			ChildInfo->ParentActorList.PrepareForWrite();
+			ChildInfo->ParentActorList.ConditionalAdd(Parent);
+		}
+
+		if (bChildIsAlreadyDependant || bChildHadParentAlready)
+		{
+			UE_LOG(LogReplicationGraph, Warning, TEXT("FGlobalActorReplicationInfoMap::AddDependentActor Child %s - Parent %s | Child already dependant of parent: %d | Child previously had parent in list: %d"),
+				   *GetPathNameSafe(Child), *GetPathNameSafe(Parent),
+				   bChildIsAlreadyDependant, bChildHadParentAlready);
+		}
+	}
+}
+
diff --git a/Engine/Plugins/Runtime/ReplicationGraph/Source/Public/ReplicationGraphTypes.h b/Engine/Plugins/Runtime/ReplicationGraph/Source/Public/ReplicationGraphTypes.h
--- a/Engine/Plugins/Runtime/ReplicationGraph/Source/Public/ReplicationGraphTypes.h
+++ b/Engine/Plugins/Runtime/ReplicationGraph/Source/Public/ReplicationGraphTypes.h
@@ -742,43 +742,7 @@ struct FGlobalActorReplicationInfoMap
 		WarnAlreadyDependant = 1 << 1,
 	};
 
-	void AddDependentActor(AActor* Parent, AActor* Child, FGlobalActorReplicationInfoMap::EWarnFlag WarnFlag=FGlobalActorReplicationInfoMap::None)
-	{
-		const bool bIsParentValid = ensureMsgf(Parent && IsActorValidForReplication(Parent), TEXT("FGlobalActorReplicationInfoMap::AddDependentActor Invalid Parent! %s"),
-			*GetPathNameSafe(Parent));
-
-		const bool bIsChildValid = ensureMsgf(Child && IsActorValidForReplication(Child), TEXT("FGlobalActorReplicationInfoMap::AddDependentActor Invalid Child! %s"),
-			*GetPathNameSafe(Child));
-
-		if (bIsParentValid && bIsChildValid)
-		{
-			const bool bDoWarnings = (WarnFlag & WarnAlreadyDependant) != 0;
-
-			bool bChildIsAlreadyDependant(false);
-			if (FGlobalActorReplicationInfo* ParentInfo = Find(Parent))
-			{
-				bChildIsAlreadyDependant = bDoWarnings && ParentInfo->DependentActorList.IsValid() && ParentInfo->DependentActorList.Contains(Child);
-
-				ParentInfo->DependentActorList.PrepareForWrite();
-				ParentInfo->DependentActorList.ConditionalAdd(Child);
-			}
-
-			bool bChildHadParentAlready(false);
-			if (FGlobalActorReplicationInfo* ChildInfo = Find(Child))
-			{
-				bChildHadParentAlready = bDoWarnings && ChildInfo->ParentActorList.IsValid() && ChildInfo->ParentActorList.Contains(Parent);
-				ChildInfo->ParentActorList.PrepareForWrite();
-				ChildInfo->ParentActorList.ConditionalAdd(Parent);
-			}
-
-			if (bChildIsAlreadyDependant || bChildHadParentAlready)
-			{
-				UE_LOG(LogReplicationGraph, Warning, TEXT("FGlobalActorReplicationInfoMap::AddDependentActor Child %s - Parent %s | Child already dependant of parent: %d | Child previously had parent in list: %d"), 
-					   *GetPathNameSafe(Child), *GetPathNameSafe(Parent),
-					   bChildIsAlreadyDependant, bChildHadParentAlready);
-			}
-		}
-	}
+	REPLICATIONGRAPH_API void AddDependentActor(AActor* Parent, AActor* Child, FGlobalActorReplicationInfoMap::EWarnFlag WarnFlag = FGlobalActorReplicationInfoMap::None);
 
 	void RemoveDependentActor(AActor* Parent, AActor* Child)
 	{
