Fix server crash due to LowestDistanceViewer being null when no viewers are close to an object that fell through the world.

diff --git a/Engine/Plugins/Runtime/ReplicationGraph/Source/Private/ReplicationGraph.cpp b/Engine/Plugins/Runtime/ReplicationGraph/Source/Private/ReplicationGraph.cpp
--- a/Engine/Plugins/Runtime/ReplicationGraph/Source/Private/ReplicationGraph.cpp
+++ b/Engine/Plugins/Runtime/ReplicationGraph/Source/Private/ReplicationGraph.cpp
@@ -3320,11 +3320,11 @@ FORCEINLINE void UReplicationGraphNode_DynamicSpatialFrequency::CalcFrequencyFor
 	float SmallestDistanceToActorSq = TNumericLimits<float>::Max();
 	const FNetViewer* LowestDistanceViewer = nullptr;
 
-	// Find the closest viewer to this item
+	// Find the closest viewer to this item or the first viewer if there are no viewers closer.
 	for (const FNetViewer& CurViewer : Viewers)
 	{
 		float CurDistance = (GlobalInfo.WorldLocation - CurViewer.ViewLocation).SizeSquared();
-		if (CurDistance < SmallestDistanceToActorSq)
+		if (LowestDistanceViewer == nullptr || CurDistance < SmallestDistanceToActorSq)
 		{
 			LowestDistanceViewer = &CurViewer;
 			SmallestDistanceToActorSq = CurDistance;
