Ensure a connection's PlayerController and current view target (pawn) are processed first in the RepGraph gather list.

diff --git a/Engine/Plugins/Runtime/ReplicationGraph/Source/Private/ReplicationGraph.cpp b/Engine/Plugins/Runtime/ReplicationGraph/Source/Private/ReplicationGraph.cpp
--- a/Engine/Plugins/Runtime/ReplicationGraph/Source/Private/ReplicationGraph.cpp
+++ b/Engine/Plugins/Runtime/ReplicationGraph/Source/Private/ReplicationGraph.cpp
@@ -99,6 +99,8 @@ static FAutoConsoleVariableRef CVarRepGraphPrintTrackClassReplication(TEXT("Net.
 int32 CVar_RepGraph_DormantDynamicActorsDestruction = 0;
 static FAutoConsoleVariableRef CVarRepGraphDormantDynamicActorsDestruction(TEXT("Net.RepGraph.DormantDynamicActorsDestruction"), CVar_RepGraph_DormantDynamicActorsDestruction, TEXT(""), ECVF_Default );
 
+static TAutoConsoleVariable<float> CVar_ForceConnectionViewerPriority(TEXT("Net.RepGraph.ForceConnectionViewerPriority"), 1, TEXT("Force the connection's player controller and viewing pawn as topmost priority."));
+
 REPGRAPH_DEVCVAR_SHIPCONST(int32, "Net.RepGraph.LogNetDormancyDetails", CVar_RepGraph_LogNetDormancyDetails, 0, "Logs actors that are removed from the replication graph/nodes.");
 REPGRAPH_DEVCVAR_SHIPCONST(int32, "Net.RepGraph.LogActorRemove", CVar_RepGraph_LogActorRemove, 0, "Logs actors that are removed from the replication graph/nodes.");
 REPGRAPH_DEVCVAR_SHIPCONST(int32, "Net.RepGraph.LogActorAdd", CVar_RepGraph_LogActorAdd, 0, "Logs actors that are added to replication graph/nodes.");
@@ -1237,19 +1239,6 @@ void UReplicationGraph::ReplicateActorListsForConnections_Default(UNetReplicatio
 					AccumulatedPriority -= 1.5f;
 				}
 
-				// -------------------
-				// Always prioritize the connection's owner and view target, since these are the most important actors for the client.
-				// -------------------
-				for (const FNetViewer& CurViewer : Viewers)
-				{
-					// We need to find if this is anyone's viewer or viewtarget, not just the parent connection.
-					if (Actor == CurViewer.ViewTarget || Actor == CurViewer.InViewer)
-					{
-						AccumulatedPriority -= 10.0f;
-						break;
-					}
-				}
-
 				// -------------------
 				//	Game code priority
 				// -------------------
@@ -1263,6 +1252,26 @@ void UReplicationGraph::ReplicateActorListsForConnections_Default(UNetReplicatio
 					{
 						DebugDetails->GameCodeScaling = -1.f;
 					}
+                }
+				
+				// -------------------
+				// Always prioritize the connection's owner and view target, since these are the most important actors for the client.
+				// -------------------
+				for (const FNetViewer& CurViewer : Viewers)
+				{
+					// We need to find if this is anyone's viewer or viewtarget, not just the parent connection.
+					if (Actor == CurViewer.ViewTarget || Actor == CurViewer.InViewer)
+					{
+						if (CVar_ForceConnectionViewerPriority.GetValueOnAnyThread() > 0)
+						{
+							AccumulatedPriority = MIN_flt;
+						}
+						else
+						{
+							AccumulatedPriority -= 10.0f;
+						}
+						break;
+					}
 				}
 
 				SortingArray->Emplace(FPrioritizedRepList::FItem(AccumulatedPriority, Actor, &GlobalData, &ConnectionData));
