Replace MIN_FLT with actual smallest negative float -MAX_FLT

diff --git a/Engine/Plugins/Runtime/ReplicationGraph/Source/Private/ReplicationGraph.cpp b/Engine/Plugins/Runtime/ReplicationGraph/Source/Private/ReplicationGraph.cpp
--- a/Engine/Plugins/Runtime/ReplicationGraph/Source/Private/ReplicationGraph.cpp
+++ b/Engine/Plugins/Runtime/ReplicationGraph/Source/Private/ReplicationGraph.cpp
@@ -1264,7 +1264,7 @@ void UReplicationGraph::ReplicateActorListsForConnections_Default(UNetReplicatio
 					{
 						if (CVar_ForceConnectionViewerPriority.GetValueOnAnyThread() > 0)
 						{
-							AccumulatedPriority = MIN_flt;
+							AccumulatedPriority = -MAX_FLT;
 						}
 						else
 						{
