ReplicationGraph will now prevent dependant actors from getting added twice to the same parent.

[RN] engine/network

diff --git a/Engine/Plugins/Runtime/ReplicationGraph/Source/Private/ReplicationGraphTypes.cpp b/Engine/Plugins/Runtime/ReplicationGraph/Source/Private/ReplicationGraphTypes.cpp
--- a/Engine/Plugins/Runtime/ReplicationGraph/Source/Private/ReplicationGraphTypes.cpp
+++ b/Engine/Plugins/Runtime/ReplicationGraph/Source/Private/ReplicationGraphTypes.cpp
@@ -469,30 +469,33 @@ void FGlobalActorReplicationInfoMap::AddDependentActor(AActor* Parent, AActor* C
 
 	if (bIsParentValid && bIsChildValid)
 	{
-		const bool bDoWarnings = (WarnFlag & WarnAlreadyDependant) != 0;
+		const bool bDoWarnings = EnumHasAnyFlags(WarnFlag, EWarnFlag::WarnAlreadyDependant);
 
 		bool bChildIsAlreadyDependant(false);
 		if (FGlobalActorReplicationInfo* ParentInfo = Find(Parent))
 		{
-			bChildIsAlreadyDependant = bDoWarnings && ParentInfo->DependentActorList.IsValid() && ParentInfo->DependentActorList.Contains(Child);
-
-			ParentInfo->DependentActorList.PrepareForWrite();
-			ParentInfo->DependentActorList.ConditionalAdd(Child);
+			bChildIsAlreadyDependant = ParentInfo->DependentActorList.IsValid() && ParentInfo->DependentActorList.Contains(Child);
+			if (bChildIsAlreadyDependant == false)
+			{
+				ParentInfo->DependentActorList.PrepareForWrite();
+				ParentInfo->DependentActorList.ConditionalAdd(Child);
+			}
 		}
 
 		bool bChildHadParentAlready(false);
 		if (FGlobalActorReplicationInfo* ChildInfo = Find(Child))
 		{
-			bChildHadParentAlready = bDoWarnings && ChildInfo->ParentActorList.IsValid() && ChildInfo->ParentActorList.Contains(Parent);
-			ChildInfo->ParentActorList.PrepareForWrite();
-			ChildInfo->ParentActorList.ConditionalAdd(Parent);
+			bChildHadParentAlready = ChildInfo->ParentActorList.IsValid() && ChildInfo->ParentActorList.Contains(Parent);
+			if (bChildHadParentAlready == false)
+			{
+				ChildInfo->ParentActorList.PrepareForWrite();
+				ChildInfo->ParentActorList.ConditionalAdd(Parent);
+			}
 		}
 
-		if (bChildIsAlreadyDependant || bChildHadParentAlready)
+		if (bDoWarnings && (bChildIsAlreadyDependant || bChildHadParentAlready))
 		{
-			UE_LOG(LogReplicationGraph, Warning, TEXT("FGlobalActorReplicationInfoMap::AddDependentActor Child %s - Parent %s | Child already dependant of parent: %d | Child previously had parent in list: %d"),
-				   *GetPathNameSafe(Child), *GetPathNameSafe(Parent),
-				   bChildIsAlreadyDependant, bChildHadParentAlready);
+			UE_LOG(LogReplicationGraph, Warning, TEXT("FGlobalActorReplicationInfoMap::AddDependentActor child %s already dependant of parent %s"), *GetNameSafe(Child), *GetNameSafe(Parent));
 		}
 	}
 }
diff --git a/Engine/Plugins/Runtime/ReplicationGraph/Source/Public/ReplicationGraphTypes.h b/Engine/Plugins/Runtime/ReplicationGraph/Source/Public/ReplicationGraphTypes.h
--- a/Engine/Plugins/Runtime/ReplicationGraph/Source/Public/ReplicationGraphTypes.h
+++ b/Engine/Plugins/Runtime/ReplicationGraph/Source/Public/ReplicationGraphTypes.h
@@ -736,13 +736,13 @@ struct FGlobalActorReplicationInfoMap
 		ClassMap.CountBytes(Ar);
 	}
 
-	enum EWarnFlag
+	enum class EWarnFlag : uint8
 	{
 		None = 0,
-		WarnAlreadyDependant = 1 << 1,
+		WarnAlreadyDependant = 1 << 0,
 	};
 
-	REPLICATIONGRAPH_API void AddDependentActor(AActor* Parent, AActor* Child, FGlobalActorReplicationInfoMap::EWarnFlag WarnFlag = FGlobalActorReplicationInfoMap::None);
+	REPLICATIONGRAPH_API void AddDependentActor(AActor* Parent, AActor* Child, FGlobalActorReplicationInfoMap::EWarnFlag WarnFlag = FGlobalActorReplicationInfoMap::EWarnFlag::None);
 
 	void RemoveDependentActor(AActor* Parent, AActor* Child)
 	{
@@ -805,6 +805,8 @@ private:
 	TClassMap<FClassReplicationInfo> ClassMap;
 };
 
+ENUM_CLASS_FLAGS(FGlobalActorReplicationInfoMap::EWarnFlag);
+
 /** Per-Actor data that is stored per connection */
 struct FConnectionReplicationActorInfo
 {
