-Now cleaning the destroyed nodes in the ReplicationGraph's AllChildNodes array after a spatial rebuild event
-Now cleaning the AllChilldNode in TearDown
-Using RemoveSingleSwap to speed up some TArray item removal operations in the RepGraph
-Added option to keep node order intact when removing destroyed nodes

[RN] engine/network

diff --git a/Engine/Plugins/Runtime/ReplicationGraph/Source/Private/ReplicationGraph.cpp b/Engine/Plugins/Runtime/ReplicationGraph/Source/Private/ReplicationGraph.cpp
--- a/Engine/Plugins/Runtime/ReplicationGraph/Source/Private/ReplicationGraph.cpp
+++ b/Engine/Plugins/Runtime/ReplicationGraph/Source/Private/ReplicationGraph.cpp
@@ -2333,7 +2333,7 @@ void UNetReplicationGraphConnection::AddConnectionGraphNode(UReplicationGraphNod
 
 void UNetReplicationGraphConnection::RemoveConnectionGraphNode(UReplicationGraphNode* Node)
 {
-	ConnectionGraphNodes.Remove(Node);
+	ConnectionGraphNodes.RemoveSingleSwap(Node);
 }
 
 bool UNetReplicationGraphConnection::PrepareForReplication()
@@ -2567,17 +2567,45 @@ void UReplicationGraphNode::NotifyResetAllNetworkActors()
 	}
 }
 
-void UReplicationGraphNode::RemoveChildNode(UReplicationGraphNode* ChildNode)
+void UReplicationGraphNode::RemoveChildNode(UReplicationGraphNode* ChildNode, UReplicationGraphNode::NodeOrdering NodeOrder)
 {
 	ensure(ChildNode != nullptr);
 
-	int32 Removed = AllChildNodes.Remove(ChildNode);
+	int32 Removed(0);
+	
+	if (NodeOrder == NodeOrdering::IgnoreOrdering)
+	{
+		Removed = AllChildNodes.RemoveSingleSwap(ChildNode, false);
+	}
+	else
+	{
+		Removed = AllChildNodes.RemoveSingle(ChildNode);
+	}
+
 	if (Removed > 0)
 	{
 		ChildNode->TearDown();
 	}
 }
 
+void UReplicationGraphNode::CleanChildNodes(UReplicationGraphNode::NodeOrdering NodeOrder)
+{
+	auto RemoveFunc = [](UReplicationGraphNode* GridChildNode)
+	{
+		return (GridChildNode == nullptr) || GridChildNode->IsPendingKill();
+	};
+
+	if (NodeOrder == NodeOrdering::IgnoreOrdering)
+	{
+		AllChildNodes.RemoveAllSwap(RemoveFunc, false);
+	}
+	else
+	{
+		AllChildNodes.RemoveAll(RemoveFunc);
+	}
+
+}
+
 void UReplicationGraphNode::TearDown()
 {
 	for (UReplicationGraphNode* Node : AllChildNodes)
@@ -2585,6 +2613,8 @@ void UReplicationGraphNode::TearDown()
 		Node->TearDown();
 	}
 
+	AllChildNodes.Reset();
+
 	MarkPendingKill();
 }
 
@@ -4786,6 +4816,8 @@ void UReplicationGraphNode_GridSpatialization2D::PrepareForReplication()
 		{
 			GEngine->ForceGarbageCollection(true);
 		}
+
+		CleanChildNodes(NodeOrdering::IgnoreOrdering);
 		
 		for (auto& MapIt : DynamicSpatializedActors)
 		{
diff --git a/Engine/Plugins/Runtime/ReplicationGraph/Source/Public/ReplicationGraph.h b/Engine/Plugins/Runtime/ReplicationGraph/Source/Public/ReplicationGraph.h
--- a/Engine/Plugins/Runtime/ReplicationGraph/Source/Public/ReplicationGraph.h
+++ b/Engine/Plugins/Runtime/ReplicationGraph/Source/Public/ReplicationGraph.h
@@ -114,7 +114,18 @@ public:
 		return NewNode;
 	}
 
-	void RemoveChildNode(UReplicationGraphNode* OutChildNode);
+	/** Node removal behavior */
+	enum class NodeOrdering
+	{
+		IgnoreOrdering, // Use faster removal but may break node processing order
+		KeepOrder,		// Use slower removal but keep the node order intact
+	};
+
+	/** Remove a child node from our list and flag it for destruction */
+	void RemoveChildNode(UReplicationGraphNode* OutChildNode, UReplicationGraphNode::NodeOrdering NodeOrder=UReplicationGraphNode::NodeOrdering::IgnoreOrdering);
+
+	/** Remove all null and about to be destroyed nodes from our list */
+	void CleanChildNodes(UReplicationGraphNode::NodeOrdering NodeOrder);
 
 protected:
 
