Disable dormant dynamic actor destruction in the replication graph. Was causing too much performance overhead.

diff --git a/Engine/Plugins/Runtime/ReplicationGraph/Source/Private/ReplicationGraph.cpp b/Engine/Plugins/Runtime/ReplicationGraph/Source/Private/ReplicationGraph.cpp
--- a/Engine/Plugins/Runtime/ReplicationGraph/Source/Private/ReplicationGraph.cpp
+++ b/Engine/Plugins/Runtime/ReplicationGraph/Source/Private/ReplicationGraph.cpp
@@ -88,7 +88,7 @@ static FAutoConsoleVariableRef CVarRepGraphTrackClassReplication(TEXT("Net.RepGr
 int32 CVar_RepGraph_PrintTrackClassReplication = 0;
 static FAutoConsoleVariableRef CVarRepGraphPrintTrackClassReplication(TEXT("Net.RepGraph.PrintTrackClassReplication"), CVar_RepGraph_PrintTrackClassReplication, TEXT(""), ECVF_Default );
 
-int32 CVar_RepGraph_DormantDynamicActorsDestruction = 1;
+int32 CVar_RepGraph_DormantDynamicActorsDestruction = 0;
 static FAutoConsoleVariableRef CVarRepGraphDormantDynamicActorsDestruction(TEXT("Net.RepGraph.DormantDynamicActorsDestruction"), CVar_RepGraph_DormantDynamicActorsDestruction, TEXT(""), ECVF_Default );
 
 REPGRAPH_DEVCVAR_SHIPCONST(int32, "Net.RepGraph.LogNetDormancyDetails", CVar_RepGraph_LogNetDormancyDetails, 0, "Logs actors that are removed from the replication graph/nodes.");
