Change channel logging to be more information in RepGraph.

diff --git a/Engine/Plugins/Runtime/ReplicationGraph/Source/Private/ReplicationGraph.cpp b/Engine/Plugins/Runtime/ReplicationGraph/Source/Private/ReplicationGraph.cpp
--- a/Engine/Plugins/Runtime/ReplicationGraph/Source/Private/ReplicationGraph.cpp
+++ b/Engine/Plugins/Runtime/ReplicationGraph/Source/Private/ReplicationGraph.cpp
@@ -800,7 +800,8 @@ int32 UReplicationGraph::ServerReplicateActors(float DeltaSeconds)
 					FConnectionReplicationActorInfo& ConnectionActorInfo = *MapIt.Value().Get();
 					UActorChannel* Channel = MapIt.Key();
 					checkSlow(Channel != nullptr);
-					ensureMsgf(Channel == ConnectionActorInfo.Channel, TEXT("Channel: %s ConnectionActorInfo.Channel: %s. Actor: %s "), *GetNameSafe(Channel), *GetNameSafe(ConnectionActorInfo.Channel), ConnectionActorInfo.Channel ? *GetNameSafe(ConnectionActorInfo.Channel->Actor) : TEXT("NULLCHANNEL"));
+					checkSlow(ConnectionActorInfo.Channel != nullptr);
+					ensureMsgf(Channel == ConnectionActorInfo.Channel, TEXT("Channel: %s ConnectionActorInfo.Channel: %s."), *(Channel ? Channel->Describe() : FString(TEXT("None"))), *(ConnectionActorInfo.Channel ? ConnectionActorInfo.Channel->Describe() : FString(TEXT("None"))));
 
 					if (ConnectionActorInfo.ActorChannelCloseFrameNum > 0 && ConnectionActorInfo.ActorChannelCloseFrameNum <= FrameNum)
 					{
@@ -825,7 +826,7 @@ int32 UReplicationGraph::ServerReplicateActors(float DeltaSeconds)
 							//UE_CLOG(DebugConnection, LogReplicationGraph, Display, TEXT("Closing Actor Channel:0x%x 0x%X0x%X, %s %d <= %d"), ConnectionActorInfo.Channel, Actor, NetConnection, *GetNameSafe(ConnectionActorInfo.Channel->Actor), ConnectionActorInfo.ActorChannelCloseFrameNum, FrameNum);
 							if (RepGraphConditionalActorBreakpoint(Actor, NetConnection))
 							{
-								UE_LOG(LogReplicationGraph, Display, TEXT("Closing Actor Channel due to timeout: %s. %d <= %d (%s)"), *Actor->GetName(), ConnectionActorInfo.ActorChannelCloseFrameNum, FrameNum, *NetConnection->Describe());
+								UE_LOG(LogReplicationGraph, Display, TEXT("Closing Actor Channel due to timeout: %s. %d <= %d (%s)"), *(ConnectionActorInfo.Channel->Describe()), ConnectionActorInfo.ActorChannelCloseFrameNum, FrameNum, *NetConnection->Describe());
 							}
 
 							INC_DWORD_STAT_BY( STAT_NetActorChannelsClosed, 1 );
@@ -1549,7 +1550,7 @@ void UReplicationGraph::UpdateActorChannelCloseFrameNum(AActor* Actor, FConnecti
 {
 	if (RepGraphConditionalActorBreakpoint(Actor, NetConnection))
 	{
-		UE_LOG(LogReplicationGraph, Display, TEXT("UReplicationGraph::UpdateActorChannelCloseFrameNum: %s. Channel: %p FrameNum: %d ActorChannelFrameTimeout: %d."), *Actor->GetName(), ConnectionData.Channel, FrameNum, GlobalData.Settings.ActorChannelFrameTimeout);
+		UE_LOG(LogReplicationGraph, Display, TEXT("UReplicationGraph::UpdateActorChannelCloseFrameNum: %s. Channel: %s FrameNum: %d ActorChannelFrameTimeout: %d."), *Actor->GetName(), *(ConnectionData.Channel ? ConnectionData.Channel->Describe() : FString(TEXT("None"))), FrameNum, GlobalData.Settings.ActorChannelFrameTimeout);
 	}
 
 	// Only update if the actor has a timeout set
diff --git a/Engine/Plugins/Runtime/ReplicationGraph/Source/Private/ReplicationGraphDebugging.cpp b/Engine/Plugins/Runtime/ReplicationGraph/Source/Private/ReplicationGraphDebugging.cpp
--- a/Engine/Plugins/Runtime/ReplicationGraph/Source/Private/ReplicationGraphDebugging.cpp
+++ b/Engine/Plugins/Runtime/ReplicationGraph/Source/Private/ReplicationGraphDebugging.cpp
@@ -17,7 +17,6 @@
 #include "ProfilingDebugging/CsvProfiler.h"
 #include "EngineUtils.h"
 #include "Engine/NetConnection.h"
-#include "Engine/ActorChannel.h"
 
 /**
  *	
@@ -959,7 +958,7 @@ void FGlobalActorReplicationInfo::LogDebugString(FOutputDevice& Ar) const
 
 void FConnectionReplicationActorInfo::LogDebugString(FOutputDevice& Ar) const
 {
-	Ar.Logf(TEXT("  Channel: %s"), *GetNameSafe(Channel));
+	Ar.Logf(TEXT("  Channel: %s"), *(Channel ? Channel->Describe() : TEXT("None")));
 	Ar.Logf(TEXT("  CullDistSq: %.2f (%.2f)"), CullDistanceSquared, FMath::Sqrt(CullDistanceSquared));
 	Ar.Logf(TEXT("  NextReplicationFrameNum: %d. ReplicationPeriodFrame: %d. LastRepFrameNum: %d. ActorChannelCloseFrameNum: %d. IsDormantOnConnection: %d. TearOff: %d"), NextReplicationFrameNum, ReplicationPeriodFrame, LastRepFrameNum, ActorChannelCloseFrameNum, bDormantOnConnection, bTearOff);
 }
