Crashfix for BasicReplicationGraph

diff --git a/Engine/Plugins/Runtime/ReplicationGraph/Source/Private/BasicReplicationGraph.cpp b/Engine/Plugins/Runtime/ReplicationGraph/Source/Private/BasicReplicationGraph.cpp
--- a/Engine/Plugins/Runtime/ReplicationGraph/Source/Private/BasicReplicationGraph.cpp
+++ b/Engine/Plugins/Runtime/ReplicationGraph/Source/Private/BasicReplicationGraph.cpp
@@ -112,7 +112,7 @@ void UBasicReplicationGraph::RouteRemoveNetworkActorToNodes(const FNewReplicated
 	}
 	else if (ActorInfo.Actor->bOnlyRelevantToOwner)
 	{
-		if (UReplicationGraphNode* Node = GetAlwaysRelevantNodeForConnection(ActorInfo.Actor->GetNetConnection()))
+		if (UReplicationGraphNode* Node = ActorInfo.Actor->GetNetConnection() ? GetAlwaysRelevantNodeForConnection(ActorInfo.Actor->GetNetConnection()) : nullptr)
 		{
 			Node->NotifyRemoveNetworkActor(ActorInfo);
 		}
