RepGraph fixes:
-Make ForceNetUpdate decrement AccumulatedPriority (more likely to replicate).
-Starvation scaling properly takes StarvationPriorityScale into account.
-UReplicationGraphNode_ActorList::NotifyRemoveNetworkActor properly sets/returns bRemovedSomething

Reference:
https://udn.unrealengine.com/questions/457755/possible-replication-graph-bug.html

diff --git a/Engine/Plugins/Runtime/ReplicationGraph/Source/Private/ReplicationGraph.cpp b/Engine/Plugins/Runtime/ReplicationGraph/Source/Private/ReplicationGraph.cpp
--- a/Engine/Plugins/Runtime/ReplicationGraph/Source/Private/ReplicationGraph.cpp
+++ b/Engine/Plugins/Runtime/ReplicationGraph/Source/Private/ReplicationGraph.cpp
@@ -904,8 +904,9 @@ void UReplicationGraph::ReplicateActorListsForConnection_Default(UNetReplication
 				// -------------------
 				if (GlobalData.Settings.StarvationPriorityScale > 0.f)
 				{
-					const uint32 FramesSinceLastRep = (FrameNum - ConnectionData.LastRepFrameNum);
-					const float StarvationFactor = 1.f - FMath::Clamp<float>((float)FramesSinceLastRep / (float)MaxFramesSinceLastRep, 0.f, 1.f);
+					// StarvationPriorityScale = scale "Frames since last rep". E.g, 2.0 means treat every missed frame as if it were 2, etc.
+					const float FramesSinceLastRep = ((float)(FrameNum - ConnectionData.LastRepFrameNum)) * GlobalData.Settings.StarvationPriorityScale;
+					const float StarvationFactor = 1.f - FMath::Clamp<float>(FramesSinceLastRep / (float)MaxFramesSinceLastRep, 0.f, 1.f);
 
 					AccumulatedPriority += StarvationFactor;
 
@@ -919,19 +920,15 @@ void UReplicationGraph::ReplicateActorListsForConnection_Default(UNetReplication
 				// -------------------
 				//	Game code priority
 				// -------------------
-							
-				if (GlobalData.ForceNetUpdateFrame > 0)
+				
+				if ( GlobalData.ForceNetUpdateFrame > ConnectionData.LastRepFrameNum )
 				{
-					uint32 ForceNetUpdateDelta = GlobalData.ForceNetUpdateFrame - ConnectionData.LastRepFrameNum;
-					if ( ForceNetUpdateDelta > 0 )
-					{
-						// Note that in legacy ForceNetUpdate did not actually bump priority. This gives us a hard coded bump if we haven't replicated since the last ForceNetUpdate frame.
-						AccumulatedPriority += 1.f;
+					// Note that in legacy ForceNetUpdate did not actually bump priority. This gives us a hard coded bump if we haven't replicated since the last ForceNetUpdate frame.
+					AccumulatedPriority -= 1.f;
 
-						if (DO_REPGRAPH_DETAILS(UNLIKELY(DebugDetails)))
-						{
-							DebugDetails->GameCodeScaling = 1.f;
-						}
+					if (DO_REPGRAPH_DETAILS(UNLIKELY(DebugDetails)))
+					{
+						DebugDetails->GameCodeScaling = -1.f;
 					}
 				}
 							
@@ -1936,6 +1933,10 @@ bool UReplicationGraphNode_ActorList::NotifyRemoveNetworkActor(const FNewReplica
 		{
 			UE_LOG(LogReplicationGraph, Warning, TEXT("Attempted to remove %s from list %s but it was not found. (StreamingLevelName == NAME_None)"), *GetActorRepListTypeDebugString(ActorInfo.Actor), *GetFullName());
 		}
+		else
+		{
+			bRemovedSomething = true;
+		}
 
 		if (CVar_RepGraph_Verify)
 		{
@@ -1944,7 +1945,7 @@ bool UReplicationGraphNode_ActorList::NotifyRemoveNetworkActor(const FNewReplica
 	}
 	else
 	{
-		StreamingLevelCollection.RemoveActor(ActorInfo, bWarnIfNotFound, this);
+		bRemovedSomething = StreamingLevelCollection.RemoveActor(ActorInfo, bWarnIfNotFound, this);
 	}
 
 	return bRemovedSomething;
