The ReplicationGraph DebugActor can now be overriden

PR #5439: Repgraph Debug Actor Suggestion (Contributed by [REDACTED_PERSON])

diff --git a/Engine/Plugins/Runtime/ReplicationGraph/Source/Private/ReplicationGraph.cpp b/Engine/Plugins/Runtime/ReplicationGraph/Source/Private/ReplicationGraph.cpp
--- a/Engine/Plugins/Runtime/ReplicationGraph/Source/Private/ReplicationGraph.cpp
+++ b/Engine/Plugins/Runtime/ReplicationGraph/Source/Private/ReplicationGraph.cpp
@@ -1898,9 +1898,13 @@ void UNetReplicationGraphConnection::InitForConnection(UNetConnection* InConnect
 	InConnection->SetReplicationConnectionDriver(this);
 
 #if !(UE_BUILD_SHIPPING || UE_BUILD_TEST)
-	DebugActor = GetWorld()->SpawnActor<AReplicationGraphDebugActor>();
-	DebugActor->ConnectionManager = this;
-	DebugActor->ReplicationGraph = Cast<UReplicationGraph>(GetOuter());
+	UReplicationGraph* Graph = Cast<UReplicationGraph>(GetOuter());
+	DebugActor = Graph->CreateDebugActor();
+	if (DebugActor)
+	{
+		DebugActor->ConnectionManager = this;
+		DebugActor->ReplicationGraph = Graph;
+	}
 #endif
 
 #if 0
diff --git a/Engine/Plugins/Runtime/ReplicationGraph/Source/Private/ReplicationGraphDebugging.cpp b/Engine/Plugins/Runtime/ReplicationGraph/Source/Private/ReplicationGraphDebugging.cpp
--- a/Engine/Plugins/Runtime/ReplicationGraph/Source/Private/ReplicationGraphDebugging.cpp
+++ b/Engine/Plugins/Runtime/ReplicationGraph/Source/Private/ReplicationGraphDebugging.cpp
@@ -993,6 +993,13 @@ void UReplicationGraph::LogConnectionGraphNodes(FReplicationGraphDebugInfo& Debu
 	}
 }
 
+#if !(UE_BUILD_SHIPPING || UE_BUILD_TEST)
+AReplicationGraphDebugActor* UReplicationGraph::CreateDebugActor() const
+{
+	return GetWorld()->SpawnActor<AReplicationGraphDebugActor>();
+}
+#endif
+
 void UReplicationGraphNode::LogNode(FReplicationGraphDebugInfo& DebugInfo, const FString& NodeName) const
 {
 	DebugInfo.Log(NodeName);
diff --git a/Engine/Plugins/Runtime/ReplicationGraph/Source/Public/ReplicationGraph.h b/Engine/Plugins/Runtime/ReplicationGraph/Source/Public/ReplicationGraph.h
--- a/Engine/Plugins/Runtime/ReplicationGraph/Source/Public/ReplicationGraph.h
+++ b/Engine/Plugins/Runtime/ReplicationGraph/Source/Public/ReplicationGraph.h
@@ -812,6 +812,10 @@ public:
 
 	uint32 GetReplicationGraphFrame() const { return ReplicationGraphFrame; }
 
+#if !(UE_BUILD_SHIPPING || UE_BUILD_TEST)
+	virtual AReplicationGraphDebugActor* CreateDebugActor() const;
+#endif
+
 
 	/** Prioritization Constants: these affect how the final priority of an actor is calculated in the prioritize phase */
 	struct FPrioritizationConstants
