Clear global actor info map from replication graph when initializing for world
Prevents leaving stale actors after a seamless travel

diff --git a/Engine/Plugins/Runtime/ReplicationGraph/Source/Private/ReplicationGraph.cpp b/Engine/Plugins/Runtime/ReplicationGraph/Source/Private/ReplicationGraph.cpp
--- a/Engine/Plugins/Runtime/ReplicationGraph/Source/Private/ReplicationGraph.cpp
+++ b/Engine/Plugins/Runtime/ReplicationGraph/Source/Private/ReplicationGraph.cpp
@@ -413,6 +413,7 @@ void UReplicationGraph::InitializeActorsInWorld(UWorld* InWorld)
 void UReplicationGraph::InitializeForWorld(UWorld* World)
 {
 	ActiveNetworkActors.Reset();
+	GlobalActorReplicationInfoMap.ResetActorMap();
 
 	for (UReplicationGraphNode* Manager : GlobalGraphNodes)
 	{
diff --git a/Engine/Plugins/Runtime/ReplicationGraph/Source/Public/ReplicationGraphTypes.h b/Engine/Plugins/Runtime/ReplicationGraph/Source/Public/ReplicationGraphTypes.h
--- a/Engine/Plugins/Runtime/ReplicationGraph/Source/Public/ReplicationGraphTypes.h
+++ b/Engine/Plugins/Runtime/ReplicationGraph/Source/Public/ReplicationGraphTypes.h
@@ -648,6 +648,8 @@ struct FGlobalActorReplicationInfoMap
 
 	int32 Num() const { return ActorMap.Num(); }
 
+	FORCEINLINE void ResetActorMap() { ActorMap.Reset(); }
+
 private:
 
 	TMap<FActorRepListType, TUniquePtr<FGlobalActorReplicationInfo>> ActorMap;
