Add missing fwd declaration that broke Mac/Linux compil targets

diff --git a/Engine/Plugins/Runtime/ReplicationGraph/Source/Public/ReplicationGraph.h b/Engine/Plugins/Runtime/ReplicationGraph/Source/Public/ReplicationGraph.h
--- a/Engine/Plugins/Runtime/ReplicationGraph/Source/Public/ReplicationGraph.h
+++ b/Engine/Plugins/Runtime/ReplicationGraph/Source/Public/ReplicationGraph.h
@@ -813,7 +813,7 @@ public:
 	uint32 GetReplicationGraphFrame() const { return ReplicationGraphFrame; }
 
 #if !(UE_BUILD_SHIPPING || UE_BUILD_TEST)
-	virtual AReplicationGraphDebugActor* CreateDebugActor() const;
+	virtual class AReplicationGraphDebugActor* CreateDebugActor() const;
 #endif
 
 
