Quick improvements to RepGraph actor prioritization:
-Always prioritize connections' ViewTargets and OwningActors, since they're usually the most important and could get starved out by other actors.
-Increase the default MaxDistanceScaling from 30m to 600m, since 30m didn't have much effect.

diff --git a/Engine/Plugins/Runtime/ReplicationGraph/Source/Private/ReplicationGraph.cpp b/Engine/Plugins/Runtime/ReplicationGraph/Source/Private/ReplicationGraph.cpp
--- a/Engine/Plugins/Runtime/ReplicationGraph/Source/Private/ReplicationGraph.cpp
+++ b/Engine/Plugins/Runtime/ReplicationGraph/Source/Private/ReplicationGraph.cpp
@@ -1092,6 +1092,14 @@ void UReplicationGraph::ReplicateActorListsForConnection_Default(UNetReplication
 					AccumulatedPriority -= 1.5f;
 				}
 
+				// -------------------
+				// Always prioritize the connection's owner and view target, since these are the most important actors for the client.
+				// -------------------
+				if (Actor == Viewer.ViewTarget || Actor == Viewer.InViewer)
+				{
+					AccumulatedPriority -= 10.0f;
+				}
+
 				// -------------------
 				//	Game code priority
 				// -------------------
diff --git a/Engine/Plugins/Runtime/ReplicationGraph/Source/Public/ReplicationGraph.h b/Engine/Plugins/Runtime/ReplicationGraph/Source/Public/ReplicationGraph.h
--- a/Engine/Plugins/Runtime/ReplicationGraph/Source/Public/ReplicationGraph.h
+++ b/Engine/Plugins/Runtime/ReplicationGraph/Source/Public/ReplicationGraph.h
@@ -817,7 +817,7 @@ public:
 	/** Prioritization Constants: these affect how the final priority of an actor is calculated in the prioritize phase */
 	struct FPrioritizationConstants
 	{
-		float MaxDistanceScaling = 3000.f * 3000.f;		// Distance scaling for prioritization scales up to this distance, everything passed this distance is the same or "capped"
+		float MaxDistanceScaling = 60000.f * 60000.f;	// Distance scaling for prioritization scales up to this distance, everything passed this distance is the same or "capped"
 		uint32 MaxFramesSinceLastRep = 20;				// Time since last rep scales up to this
 		
 	};
