Prevent crash if an unexpected rpc is called while preparing a fast shared bunch

diff --git a/Engine/Plugins/Runtime/ReplicationGraph/Source/Private/ReplicationGraph.cpp b/Engine/Plugins/Runtime/ReplicationGraph/Source/Private/ReplicationGraph.cpp
--- a/Engine/Plugins/Runtime/ReplicationGraph/Source/Private/ReplicationGraph.cpp
+++ b/Engine/Plugins/Runtime/ReplicationGraph/Source/Private/ReplicationGraph.cpp
@@ -1393,6 +1393,7 @@ int64 UReplicationGraph::ReplicateSingleActor_FastShared(AActor* Actor, FConnect
 		// Make shared thing
 		FastSharedReplicationBunch = &OutBunch;
 		FastSharedReplicationChannel = ActorChannel;
+		FastSharedReplicationFuncName = GlobalActorInfo.Settings.FastSharedReplicationFuncName;
 
 		// Calling this function *should* result in an RPC call that we trap and fill out FastSharedReplicationBunch. See UReplicationGraph::ProcessRemoteFunction
 		if (GlobalActorInfo.Settings.FastSharedReplicationFunc(Actor) == false)
@@ -1400,6 +1401,7 @@ int64 UReplicationGraph::ReplicateSingleActor_FastShared(AActor* Actor, FConnect
 			// Something failed and we don't want to fast replicate. We wont check again this frame
 			FastSharedReplicationBunch = nullptr;
 			FastSharedReplicationChannel = nullptr;
+			FastSharedReplicationFuncName = NAME_None;
 			return 0;
 		}
 
@@ -1414,6 +1416,7 @@ int64 UReplicationGraph::ReplicateSingleActor_FastShared(AActor* Actor, FConnect
 			// A new bunch was not produced this frame, but there is still valid data (If FastSharedReplicationFunc returns false, there is no valid data)
 			FastSharedReplicationBunch = nullptr;
 			FastSharedReplicationChannel = nullptr;
+			FastSharedReplicationFuncName = NAME_None;
 		}
 	}
 
@@ -1652,7 +1655,7 @@ bool UReplicationGraph::ProcessRemoteFunction(class AActor* Actor, UFunction* Fu
 	// --------------------------------------------------------------------------------------------------------------------------------------------------------------------------
 	// FastShared Replication. This is ugly but the idea here is to just fill out the bunch parameters and return so that this bunch can be reused by other connections
 	// --------------------------------------------------------------------------------------------------------------------------------------------------------------------------
-	if (FastSharedReplicationBunch)
+	if (FastSharedReplicationBunch && (FastSharedReplicationFuncName == Function->GetFName()))
 	{
 		// We also cache off a channel so we can call some of the serialization functions on it. This isn't really necessary though and we could break those parts off
 		// into a static function.
@@ -1673,6 +1676,7 @@ bool UReplicationGraph::ProcessRemoteFunction(class AActor* Actor, UFunction* Fu
 			
 			FastSharedReplicationBunch = nullptr;
 			FastSharedReplicationChannel = nullptr;
+			FastSharedReplicationFuncName = NAME_None;
 		}
 		return true;
 	}
diff --git a/Engine/Plugins/Runtime/ReplicationGraph/Source/Public/ReplicationGraph.h b/Engine/Plugins/Runtime/ReplicationGraph/Source/Public/ReplicationGraph.h
--- a/Engine/Plugins/Runtime/ReplicationGraph/Source/Public/ReplicationGraph.h
+++ b/Engine/Plugins/Runtime/ReplicationGraph/Source/Public/ReplicationGraph.h
@@ -889,6 +889,7 @@ protected:
 
 	FOutBunch* FastSharedReplicationBunch = nullptr;
 	class UActorChannel* FastSharedReplicationChannel = nullptr;
+	FName FastSharedReplicationFuncName = NAME_None;
 
 #if REPGRAPH_DETAILS
 	bool bEnableFullActorPrioritizationDetailsAllConnections = false;
diff --git a/Engine/Plugins/Runtime/ReplicationGraph/Source/Public/ReplicationGraphTypes.h b/Engine/Plugins/Runtime/ReplicationGraph/Source/Public/ReplicationGraphTypes.h
--- a/Engine/Plugins/Runtime/ReplicationGraph/Source/Public/ReplicationGraphTypes.h
+++ b/Engine/Plugins/Runtime/ReplicationGraph/Source/Public/ReplicationGraphTypes.h
@@ -364,6 +364,7 @@ struct FClassReplicationInfo
 	uint8 ActorChannelFrameTimeout = 4;
 
 	TFunction<bool(AActor*)> FastSharedReplicationFunc = nullptr;
+	FName FastSharedReplicationFuncName = NAME_None;
 
 	FString BuildDebugStringDelta() const
 	{
