Creative: Prioritize fix for special case (from Josh May)

diff --git a/Engine/Plugins/Runtime/ReplicationGraph/Source/Private/ReplicationGraph.cpp b/Engine/Plugins/Runtime/ReplicationGraph/Source/Private/ReplicationGraph.cpp
--- a/Engine/Plugins/Runtime/ReplicationGraph/Source/Private/ReplicationGraph.cpp
+++ b/Engine/Plugins/Runtime/ReplicationGraph/Source/Private/ReplicationGraph.cpp
@@ -1026,7 +1026,7 @@ void UReplicationGraph::ReplicateActorListsForConnection_Default(UNetReplication
 					DO_REPGRAPH_DETAILS(DebugDetails = PrioritizedReplicationList.GetNextFullDebugDetails(Actor));
 				}
 
-				float AccumulatedPriority = 0.f;
+				float AccumulatedPriority = GlobalData.Settings.AccumulatedNetPriorityBias;
 
 				// -------------------
 				// Distance Scaling
diff --git a/Engine/Plugins/Runtime/ReplicationGraph/Source/Public/ReplicationGraphTypes.h b/Engine/Plugins/Runtime/ReplicationGraph/Source/Public/ReplicationGraphTypes.h
--- a/Engine/Plugins/Runtime/ReplicationGraph/Source/Public/ReplicationGraphTypes.h
+++ b/Engine/Plugins/Runtime/ReplicationGraph/Source/Public/ReplicationGraphTypes.h
@@ -358,6 +358,7 @@ struct FClassReplicationInfo
 	float DistancePriorityScale = 1.f;
 	float StarvationPriorityScale = 1.f;
 	float CullDistanceSquared = 0.f;
+	float AccumulatedNetPriorityBias = 0.f;
 	
 	uint8 ReplicationPeriodFrame = 1;
 	uint8 FastPath_ReplicationPeriodFrame = 1;
