Added optional budget in the RepGraph for data sent when a new ActorChannel is created on a connection. Can be used to improve actor discovery time when loading into a level with alot of dynamic actors packed into a single grid.

diff --git a/Engine/Plugins/Runtime/ReplicationGraph/Source/Private/ReplicationGraph.cpp b/Engine/Plugins/Runtime/ReplicationGraph/Source/Private/ReplicationGraph.cpp
--- a/Engine/Plugins/Runtime/ReplicationGraph/Source/Private/ReplicationGraph.cpp
+++ b/Engine/Plugins/Runtime/ReplicationGraph/Source/Private/ReplicationGraph.cpp
@@ -179,6 +179,7 @@ UReplicationGraph::UReplicationGraph()
 {
 	ReplicationConnectionManagerClass = UNetReplicationGraphConnection::StaticClass();
 	GlobalActorChannelFrameNumTimeout = 2;
+	ActorDiscoveryMaxBitsPerFrame = 0;
 
 	if (!HasAnyFlags(RF_ClassDefaultObject))
 	{
@@ -785,6 +786,7 @@ int32 UReplicationGraph::ServerReplicateActors(float DeltaSeconds)
 		}
 
 		FBitWriter& ConnectionSendBuffer = NetConnection->SendBuffer;
+		ConnectionManager->QueuedBitsForActorDiscovery = 0;
 
 		// --------------------------------------------------------------------------------------------------------------
 		// GATHER list of ReplicationLists for this connection
@@ -900,7 +902,7 @@ int32 UReplicationGraph::ServerReplicateActors(float DeltaSeconds)
 				{
 					FGlobalActorReplicationInfo& GlobalInfo = GlobalActorReplicationInfoMap.Get(ConnectionManager->DebugActor);
 					FConnectionReplicationActorInfo& ActorInfo = ConnectionActorInfoMap.FindOrAdd(ConnectionManager->DebugActor);
-					ReplicateSingleActor(ConnectionManager->DebugActor, ActorInfo, GlobalInfo, ConnectionActorInfoMap, NetConnection, FrameNum);
+					ReplicateSingleActor(ConnectionManager->DebugActor, ActorInfo, GlobalInfo, ConnectionActorInfoMap, *ConnectionManager, FrameNum);
 				}
 			}
 #endif
@@ -1149,7 +1151,7 @@ void UReplicationGraph::ReplicateActorListsForConnection_Default(UNetReplication
 
 			FGlobalActorReplicationInfo& GlobalActorInfo = *RepItem.GlobalData;
 
-			int64 BitsWritten = ReplicateSingleActor(Actor, ActorInfo, GlobalActorInfo, ConnectionActorInfoMap, NetConnection, FrameNum);
+			int64 BitsWritten = ReplicateSingleActor(Actor, ActorInfo, GlobalActorInfo, ConnectionActorInfoMap, *ConnectionManager, FrameNum);
 
 			// --------------------------------------------------
 			//	Update Packet Budget Tracking
@@ -1316,7 +1318,7 @@ void UReplicationGraph::ReplicateActorListsForConnection_FastShared(UNetReplicat
 				continue;
 			}
 
-			BitsWritten = (int32)ReplicateSingleActor_FastShared(Actor, ConnectionData, GlobalActorInfo, NetConnection, FrameNum);
+			BitsWritten = (int32)ReplicateSingleActor_FastShared(Actor, ConnectionData, GlobalActorInfo, *ConnectionManager, FrameNum);
 			TotalBitsWritten += BitsWritten;
 
 #if !(UE_BUILD_SHIPPING || UE_BUILD_TEST)
@@ -1338,8 +1340,10 @@ void UReplicationGraph::ReplicateActorListsForConnection_FastShared(UNetReplicat
 
 REPGRAPH_DEVCVAR_SHIPCONST(int32, "Net.RepGraph.FastShared.ForceFull", CVar_RepGraph_FastShared_ForceFull, 0, "Redirects calls to ReplicateSingleActor_FastShared to ReplicateSingleActor");
 
-int64 UReplicationGraph::ReplicateSingleActor_FastShared(AActor* Actor, FConnectionReplicationActorInfo& ConnectionData, FGlobalActorReplicationInfo& GlobalActorInfo, UNetConnection* NetConnection, const uint32 FrameNum)
+int64 UReplicationGraph::ReplicateSingleActor_FastShared(AActor* Actor, FConnectionReplicationActorInfo& ConnectionData, FGlobalActorReplicationInfo& GlobalActorInfo, UNetReplicationGraphConnection& ConnectionManager, const uint32 FrameNum)
 {
+	UNetConnection* NetConnection = ConnectionManager.NetConnection;
+
 	// No matter what we consider this FastShared rep to happen. Even if the actor doesn't produce a bunch or its empty or stale, etc. We still consider this replication to have happened
 	// for high level frequency purposes (E.g, UReplicationGraphNode_DynamicSpatialFrequency). But we want to do the update at the end of this function, not at the top since it can early out
 	// if the actor doesn't produce a new bunch and this connection already got the last bunch produced.
@@ -1351,7 +1355,7 @@ int64 UReplicationGraph::ReplicateSingleActor_FastShared(AActor* Actor, FConnect
 	
 	if (CVar_RepGraph_FastShared_ForceFull > 0)
 	{
-		return ReplicateSingleActor(Actor, ConnectionData, GlobalActorInfo, FindOrAddConnectionManager(NetConnection)->ActorInfoMap, NetConnection, FrameNum);
+		return ReplicateSingleActor(Actor, ConnectionData, GlobalActorInfo, FindOrAddConnectionManager(NetConnection)->ActorInfoMap, ConnectionManager, FrameNum);
 	}
 
 	int32 BitsWritten = 0;
@@ -1451,10 +1455,12 @@ int64 UReplicationGraph::ReplicateSingleActor_FastShared(AActor* Actor, FConnect
 	return OutBunch.GetNumBits();
 }
 
-int64 UReplicationGraph::ReplicateSingleActor(AActor* Actor, FConnectionReplicationActorInfo& ActorInfo, FGlobalActorReplicationInfo& GlobalActorInfo, FPerConnectionActorInfoMap& ConnectionActorInfoMap, UNetConnection* NetConnection, const uint32 FrameNum)
+int64 UReplicationGraph::ReplicateSingleActor(AActor* Actor, FConnectionReplicationActorInfo& ActorInfo, FGlobalActorReplicationInfo& GlobalActorInfo, FPerConnectionActorInfoMap& ConnectionActorInfoMap, UNetReplicationGraphConnection& ConnectionManager, const uint32 FrameNum)
 {
 	RG_QUICK_SCOPE_CYCLE_COUNTER(NET_ReplicateActors_ReplicateSingleActor);
 
+	UNetConnection* NetConnection = ConnectionManager.NetConnection;
+
 	if (RepGraphConditionalActorBreakpoint(Actor, NetConnection))
 	{
 		UE_LOG(LogReplicationGraph, Display, TEXT("UReplicationGraph::ReplicateSingleActor: %s. NetConnection: %s"), *Actor->GetName(), *NetConnection->Describe());
@@ -1483,7 +1489,9 @@ int64 UReplicationGraph::ReplicateSingleActor(AActor* Actor, FConnectionReplicat
 	const bool bWantsToGoDormant = GlobalActorInfo.bWantsToBeDormant;
 	TActorRepListViewBase<FActorRepList*> DependentActorList(GlobalActorInfo.DependentActorList.RepList.GetReference());
 
-	if (ActorInfo.Channel == nullptr)
+	bool bOpenActorChannel = (ActorInfo.Channel == nullptr);
+
+	if (bOpenActorChannel)
 	{
 		// Create a new channel for this actor.
 		INC_DWORD_STAT_BY( STAT_NetActorChannelsOpened, 1 );
@@ -1565,10 +1573,23 @@ int64 UReplicationGraph::ReplicateSingleActor(AActor* Actor, FConnectionReplicat
 			}
 
 			//UE_LOG(LogReplicationGraph, Display, TEXT("DependentActor %s %s. NextReplicationFrameNum: %d. FrameNum: %d. ForceNetUpdateFrame: %d. LastRepFrameNum: %d."), *DependentActor->GetPathName(), *NetConnection->GetName(), DependentActorConnectionInfo.NextReplicationFrameNum, FrameNum, DependentActorGlobalData.ForceNetUpdateFrame, DependentActorConnectionInfo.LastRepFrameNum);
-			BitsWritten += ReplicateSingleActor(DependentActor, DependentActorConnectionInfo, DependentActorGlobalData, ConnectionActorInfoMap, NetConnection, FrameNum);
+			BitsWritten += ReplicateSingleActor(DependentActor, DependentActorConnectionInfo, DependentActorGlobalData, ConnectionActorInfoMap, ConnectionManager, FrameNum);
 		}					
 	}
 
+	// Optional budget for actor discovery traffic
+	if (ActorDiscoveryMaxBitsPerFrame > 0 && bOpenActorChannel )
+	{
+		if (ConnectionManager.QueuedBitsForActorDiscovery < ActorDiscoveryMaxBitsPerFrame)
+		{
+			ConnectionManager.QueuedBitsForActorDiscovery += BitsWritten;
+
+			// Remove the discovery traffic from the regular traffic
+			NetConnection->QueuedBits -= BitsWritten;
+			BitsWritten = 0;
+		}
+	}
+
 	return BitsWritten;
 }
 
@@ -1865,6 +1886,27 @@ bool UReplicationGraph::IsConnectionReady(UNetConnection* Connection)
 	return Connection->QueuedBits + Connection->SendBuffer.GetNumBits() <= 0;
 }
 
+void UReplicationGraph::SetActorDiscoveryBudget(int32 ActorDiscoveryBudgetInKBytesPerSec)
+{
+	// Disable the seperate actor discovery budget when 0
+	if (ActorDiscoveryBudgetInKBytesPerSec <= 0)
+	{
+		ActorDiscoveryMaxBitsPerFrame = 0;
+		return;
+	}
+
+	if (NetDriver == nullptr)
+	{
+		UE_LOG(LogReplicationGraph, Warning, TEXT("SetActorDiscoveryBudget ignored since NetDriver was not initialized."));
+		return;
+	}
+
+	int32 MaxNetworkFPS = NetDriver->NetServerMaxTickRate;
+
+	ActorDiscoveryMaxBitsPerFrame = (ActorDiscoveryBudgetInKBytesPerSec * 1000 * 8) / MaxNetworkFPS;
+	UE_LOG(LogReplicationGraph, Display, TEXT("SetActorDiscoveryBudget set to %d kBps (%d bits per network tick)."), ActorDiscoveryBudgetInKBytesPerSec, ActorDiscoveryMaxBitsPerFrame);
+}
+
 // --------------------------------------------------------------------------------------------------------------------------------------------
 // --------------------------------------------------------------------------------------------------------------------------------------------
 // --------------------------------------------------------------------------------------------------------------------------------------------
@@ -2847,7 +2889,7 @@ void UReplicationGraphNode_DynamicSpatialFrequency::GatherActorListsForConnectio
 
 			if (ReadyForNextReplication(ConnectionInfo, GlobalInfo, FrameNum))
 			{
-				BitsWritten += RepGraph->ReplicateSingleActor(Actor, ConnectionInfo, GlobalInfo, ConnectionActorInfoMap, NetConnection, FrameNum);
+				BitsWritten += RepGraph->ReplicateSingleActor(Actor, ConnectionInfo, GlobalInfo, ConnectionActorInfoMap, Params.ConnectionManager, FrameNum);
 				ConnectionInfo.FastPath_LastRepFrameNum = FrameNum; // Manually update this here, so that we don't fast rep next frame. When they line up, use default replication.
 			}
 
@@ -2857,7 +2899,7 @@ void UReplicationGraphNode_DynamicSpatialFrequency::GatherActorListsForConnectio
 			
 			else if (Item.EnableFastPath && ReadyForNextReplication_FastPath(ConnectionInfo, GlobalInfo, FrameNum))
 			{
-				const int64 FastSharedBits = RepGraph->ReplicateSingleActor_FastShared(Actor, ConnectionInfo, GlobalInfo, NetConnection, FrameNum);
+				const int64 FastSharedBits = RepGraph->ReplicateSingleActor_FastShared(Actor, ConnectionInfo, GlobalInfo, Params.ConnectionManager, FrameNum);
 				QueuedBits -= FastSharedBits; // We are doing our own bandwidth limiting here, so offset the netconnection's tracking.
 				BitsWritten += FastSharedBits;
 			}
@@ -4667,4 +4709,20 @@ FAutoConsoleCommandWithWorldAndArgs NetRepGraphPrintChannelCounters(TEXT("Net.Re
 			UE_LOG(LogReplicationGraph, Display, TEXT("%s : %s ----> %d"), *GetNameSafe(Pair.Actor.Get()), *GetNameSafe(Pair.Connection.Get()), Count );
 		}
 	})
-);
\ No newline at end of file
+);
+
+// ------------------------------------------------------------------------------
+
+FAutoConsoleCommandWithWorldAndArgs ChangeActorDiscoveryBudget(TEXT("Net.RepGraph.ActorDiscoveryBudget"), TEXT("Set a separate network traffic budget for data sent when opening a new actor channel. Value in kilobytes per second"), FConsoleCommandWithWorldAndArgsDelegate::CreateLambda([](const TArray< FString >& Args, UWorld* World)
+{
+	int32 BudgetInKBPS = 0;
+	if (Args.Num() > 0)
+	{
+		LexTryParseString<int32>(BudgetInKBPS, *Args[0]);
+	}
+
+	for (TObjectIterator<UReplicationGraph> It; It; ++It)
+	{
+		It->SetActorDiscoveryBudget(BudgetInKBPS);
+	}
+}));
\ No newline at end of file
diff --git a/Engine/Plugins/Runtime/ReplicationGraph/Source/Public/ReplicationGraph.h b/Engine/Plugins/Runtime/ReplicationGraph/Source/Public/ReplicationGraph.h
--- a/Engine/Plugins/Runtime/ReplicationGraph/Source/Public/ReplicationGraph.h
+++ b/Engine/Plugins/Runtime/ReplicationGraph/Source/Public/ReplicationGraph.h
@@ -776,6 +776,8 @@ public:
 
 	bool IsConnectionReady(UNetConnection* Connection);
 
+	void SetActorDiscoveryBudget(int32 ActorDiscoveryBudgetInKBytesPerSec);
+
 	// --------------------------------------------------------------
 
 	/** Creates a new node for the graph. This and UReplicationNode::CreateChildNode should be the only things that create the graph node UObjects */
@@ -836,8 +838,8 @@ public:
 
 	// --------------------------------------------------------------
 
-	int64 ReplicateSingleActor(AActor* Actor, FConnectionReplicationActorInfo& ActorInfo, FGlobalActorReplicationInfo& GlobalActorInfo, FPerConnectionActorInfoMap& ConnectionActorInfoMap, UNetConnection* NetConnection, const uint32 FrameNum);
-	int64 ReplicateSingleActor_FastShared(AActor* Actor, FConnectionReplicationActorInfo& ConnectionData, FGlobalActorReplicationInfo& GlobalActorInfo, UNetConnection* NetConnection, const uint32 FrameNum);
+	int64 ReplicateSingleActor(AActor* Actor, FConnectionReplicationActorInfo& ActorInfo, FGlobalActorReplicationInfo& GlobalActorInfo, FPerConnectionActorInfoMap& ConnectionActorInfoMap, UNetReplicationGraphConnection& ConnectionManager, const uint32 FrameNum);
+	int64 ReplicateSingleActor_FastShared(AActor* Actor, FConnectionReplicationActorInfo& ConnectionData, FGlobalActorReplicationInfo& GlobalActorInfo, UNetReplicationGraphConnection& ConnectionManager, const uint32 FrameNum);
 
 	void UpdateActorChannelCloseFrameNum(AActor* Actor, FConnectionReplicationActorInfo& ConnectionData, const FGlobalActorReplicationInfo& GlobalData, const uint32 FrameNum, UNetConnection* NetConnection) const;
 
@@ -909,6 +911,9 @@ private:
 	/** Internal frame counter for replication. This is only updated by us. The one of UNetDriver can be updated by RPC calls and is only used to invalidate shared property CLs/serialiation data. */
 	uint32 ReplicationGraphFrame = 0;
 
+	/** Separate bandwidth cap for traffic used when opening actor channels. Ignored if set to 0 */
+	int32 ActorDiscoveryMaxBitsPerFrame;
+
 	UNetReplicationGraphConnection* CreateClientConnectionManagerInternal(UNetConnection* Connection);
 
 	friend class AReplicationGraphDebugActor;
@@ -960,6 +965,9 @@ public:
 
 	FVector LastGatherLocation;
 
+	// Nb of bits sent for actor channel creation when a dedicated budget is allocated for this
+	int32 QueuedBitsForActorDiscovery = 0;
+
 	/** Returns connection graph nodes. This is const so that you do not mutate the array itself. You should use AddConnectionGraphNode/RemoveConnectionGraphNode.  */
 	const TArray<UReplicationGraphNode*>& GetConnectionGraphNodes() { return ConnectionGraphNodes; }
 
