Clamp actor locations in UReplicationGraphNode_GridSpatialization2D::GetCellInfoForActor to the world maximum size in case we get bad values and try to allocate a massive grid later on in UReplicationGraphNode_GridSpatialization2D::GetGridNodesForActor.

diff --git a/Engine/Plugins/Runtime/ReplicationGraph/Source/Private/ReplicationGraph.cpp b/Engine/Plugins/Runtime/ReplicationGraph/Source/Private/ReplicationGraph.cpp
--- a/Engine/Plugins/Runtime/ReplicationGraph/Source/Private/ReplicationGraph.cpp
+++ b/Engine/Plugins/Runtime/ReplicationGraph/Source/Private/ReplicationGraph.cpp
@@ -3875,9 +3875,20 @@ UReplicationGraphNode_GridSpatialization2D::FActorCellInfo UReplicationGraphNode
 	}
 #endif
 
+	FVector ClampedLocation = Location3D;
+
+	// Sanity check the actor's location. If it's garbage, we could end up with a gigantic allocation in GetGridNodesForActor as we adjust the grid.
+	if (Location3D.X < -HALF_WORLD_MAX || Location3D.X > HALF_WORLD_MAX ||
+		Location3D.Y < -HALF_WORLD_MAX || Location3D.Y > HALF_WORLD_MAX ||
+		Location3D.Z < -HALF_WORLD_MAX || Location3D.Z > HALF_WORLD_MAX)
+	{
+		UE_LOG(LogReplicationGraph, Warning, TEXT("GetCellInfoForActor: Actor %s is outside world bounds with a location of %s. Clamping grid location to world bounds."), *GetFullNameSafe(Actor), *Location3D.ToString());
+		ClampedLocation = Location3D.BoundToCube(HALF_WORLD_MAX);
+	}
+
 	FActorCellInfo CellInfo;
-	const float LocationBiasX = (Location3D.X - SpatialBias.X);
-	const float LocationBiasY = (Location3D.Y - SpatialBias.Y);
+	const float LocationBiasX = (ClampedLocation.X - SpatialBias.X);
+	const float LocationBiasY = (ClampedLocation.Y - SpatialBias.Y);
 
 	const float Dist = FMath::Sqrt(CullDistanceSquared);	 // Fixme Sqrt
 	const float MinX = LocationBiasX - Dist;
