Stop using a static variable to hold the time left until the next tick

#rnx

diff --git a/Engine/Plugins/Runtime/ReplicationGraph/Source/Private/ReplicationGraph.cpp b/Engine/Plugins/Runtime/ReplicationGraph/Source/Private/ReplicationGraph.cpp
--- a/Engine/Plugins/Runtime/ReplicationGraph/Source/Private/ReplicationGraph.cpp
+++ b/Engine/Plugins/Runtime/ReplicationGraph/Source/Private/ReplicationGraph.cpp
@@ -731,13 +731,12 @@ int32 UReplicationGraph::ServerReplicateActors(float DeltaSeconds)
 #endif
 	const float TimeBetweenUpdates = TargetUpdatesPerSecond > 0 ? (1.f / (float)TargetUpdatesPerSecond) : 0.f;
 
-	static float TimeLeft = TimeBetweenUpdates;
-	TimeLeft -= DeltaSeconds;
-	if (TimeLeft > 0.f)
+	TimeLeftUntilUpdate -= DeltaSeconds;
+	if (TimeLeftUntilUpdate > 0.f)
 	{
 		return 0;
 	}
-	TimeLeft = TimeBetweenUpdates;
+	TimeLeftUntilUpdate = TimeBetweenUpdates;
 #endif
 	
 	SCOPED_NAMED_EVENT(UReplicationGraph_ServerReplicateActors, FColor::Green);
diff --git a/Engine/Plugins/Runtime/ReplicationGraph/Source/Public/ReplicationGraph.h b/Engine/Plugins/Runtime/ReplicationGraph/Source/Public/ReplicationGraph.h
--- a/Engine/Plugins/Runtime/ReplicationGraph/Source/Public/ReplicationGraph.h
+++ b/Engine/Plugins/Runtime/ReplicationGraph/Source/Public/ReplicationGraph.h
@@ -914,6 +914,9 @@ private:
 	/** Separate bandwidth cap for traffic used when opening actor channels. Ignored if set to 0 */
 	int32 ActorDiscoveryMaxBitsPerFrame;
 
+	/** Internal time used to track when the next update should occur based on frequency settings. */
+	float TimeLeftUntilUpdate = 0.f;
+
 	UNetReplicationGraphConnection* CreateClientConnectionManagerInternal(UNetConnection* Connection);
 
 	friend class AReplicationGraphDebugActor;
