Fix build failures on Mac. Export symbols that need to be for Mac/Linux

diff --git a/Engine/Plugins/Online/OnlineFramework/Source/Party/Public/Party/PartyTypes.h b/Engine/Plugins/Online/OnlineFramework/Source/Party/Public/Party/PartyTypes.h
--- a/Engine/Plugins/Online/OnlineFramework/Source/Party/Public/Party/PartyTypes.h
+++ b/Engine/Plugins/Online/OnlineFramework/Source/Party/Public/Party/PartyTypes.h
@@ -273,8 +273,8 @@ public:
 	UPROPERTY()
 	bool bOnlyLeaderFriendsCanJoin = false;
 
-	bool operator==(const FPartyPrivacySettings& Other) const;
-	bool operator!=(const FPartyPrivacySettings& Other) const { return !operator==(Other); }
+	bool PARTY_API operator==(const FPartyPrivacySettings& Other) const;
+	bool PARTY_API operator!=(const FPartyPrivacySettings& Other) const { return !operator==(Other); }
 
 	FPartyPrivacySettings() {}
 };
diff --git a/Engine/Plugins/Runtime/AudioCapture/Source/AudioCapture/Public/AudioCapture.h b/Engine/Plugins/Runtime/AudioCapture/Source/AudioCapture/Public/AudioCapture.h
--- a/Engine/Plugins/Runtime/AudioCapture/Source/AudioCapture/Public/AudioCapture.h
+++ b/Engine/Plugins/Runtime/AudioCapture/Source/AudioCapture/Public/AudioCapture.h
@@ -87,7 +87,7 @@ namespace Audio
 	{
 	public:
 		FAudioCaptureSynth();
-		virtual ~FAudioCaptureSynth();
+		AUDIOCAPTURE_API virtual ~FAudioCaptureSynth();
 
 		//~ IAudioCaptureCallback Begin
 		virtual void OnAudioCapture(float* AudioData, int32 NumFrames, int32 NumChannels, double StreamTime, bool bOverflow) override;
@@ -153,4 +153,4 @@ class FAudioCaptureModule : public IModuleInterface
 public:
 	virtual void StartupModule() override;
 	virtual void ShutdownModule() override;
-};
\ No newline at end of file
+};
diff --git a/Engine/Plugins/Runtime/ReplicationGraph/Source/Public/ReplicationGraphTypes.h b/Engine/Plugins/Runtime/ReplicationGraph/Source/Public/ReplicationGraphTypes.h
--- a/Engine/Plugins/Runtime/ReplicationGraph/Source/Public/ReplicationGraphTypes.h
+++ b/Engine/Plugins/Runtime/ReplicationGraph/Source/Public/ReplicationGraphTypes.h
@@ -31,7 +31,7 @@ DECLARE_LOG_CATEGORY_EXTERN( LogReplicationGraph, Log, All );
 	#define repCheck(x) check(x)
 	#define repCheckf(expr, format, ...) checkf(expr, format, ##__VA_ARGS__ )
 	#define RG_QUICK_SCOPE_CYCLE_COUNTER(x) QUICK_SCOPE_CYCLE_COUNTER(x)
-	extern int32 CVar_RepGraph_Verify;
+	REPLICATIONGRAPH_API extern int32 CVar_RepGraph_Verify;
 #else 
 	#define REPGRAPH_DETAILS 0
 	#define DO_REPGRAPH_DETAILS(X) 0
@@ -1435,4 +1435,4 @@ struct FActorConnectionPair
 };
 
 // Generic/global pair that can be set by debug commands etc for extra logging/debugging functionality
-extern FActorConnectionPair DebugActorConnectionPair;
\ No newline at end of file
+extern FActorConnectionPair DebugActorConnectionPair;
diff --git a/Engine/Source/Runtime/Engine/Classes/Engine/DataAsset.h b/Engine/Source/Runtime/Engine/Classes/Engine/DataAsset.h
--- a/Engine/Source/Runtime/Engine/Classes/Engine/DataAsset.h
+++ b/Engine/Source/Runtime/Engine/Classes/Engine/DataAsset.h
@@ -36,7 +36,7 @@ private:
  * To override this behavior, override GetPrimaryAssetId in your native class
  */
 UCLASS(abstract, MinimalAPI, Blueprintable)
-class UPrimaryDataAsset : public UDataAsset
+class ENGINE_VTABLE UPrimaryDataAsset : public UDataAsset
 {
 	GENERATED_BODY()
 
diff --git a/Engine/Source/Runtime/Engine/Classes/Engine/DataTable.h b/Engine/Source/Runtime/Engine/Classes/Engine/DataTable.h
--- a/Engine/Source/Runtime/Engine/Classes/Engine/DataTable.h
+++ b/Engine/Source/Runtime/Engine/Classes/Engine/DataTable.h
@@ -53,7 +53,7 @@ struct FTableRowBase
  * Imported spreadsheet table.
  */
 UCLASS(MinimalAPI, BlueprintType)
-class UDataTable
+class ENGINE_VTABLE UDataTable
 	: public UObject
 {
 	GENERATED_UCLASS_BODY()
