Fix crash when calling Net.RepGraph.Printall when SkippedActors ptr is null

diff --git a/Engine/Plugins/Runtime/ReplicationGraph/Source/Private/ReplicationGraphDebugging.cpp b/Engine/Plugins/Runtime/ReplicationGraph/Source/Private/ReplicationGraphDebugging.cpp
--- a/Engine/Plugins/Runtime/ReplicationGraph/Source/Private/ReplicationGraphDebugging.cpp
+++ b/Engine/Plugins/Runtime/ReplicationGraph/Source/Private/ReplicationGraphDebugging.cpp
@@ -1108,38 +1108,42 @@ void PrintPrioritizedList(FOutputDevice& Ar, UNetReplicationGraphConnection* Con
 	
 	// Skipped actors
 #if REPGRAPH_DETAILS
-	Ar.Logf(TEXT("[%d Skipped Actors]"), PrioritizedList->SkippedDebugDetails->Num());
 
-	FNativeClassAccumulator DormantClasses;
-	FNativeClassAccumulator CulledClasses;
-
-	for (const FSkippedActorFullDebugDetails& SkippedDetails : *PrioritizedList->SkippedDebugDetails)
+	if (PrioritizedList->SkippedDebugDetails.IsValid())
 	{
-		FString SkippedStr;
-		if (SkippedDetails.bWasDormant)
-		{
-			SkippedStr = TEXT("Dormant");
-			DormantClasses.Increment(SkippedDetails.Actor->GetClass());
-		}
-		else if (SkippedDetails.DistanceCulled > 0.f)
-		{
-			SkippedStr = FString::Printf(TEXT("Dist Culled %.2f"), SkippedDetails.DistanceCulled);
-			CulledClasses.Increment(SkippedDetails.Actor->GetClass());
-		}
-		else if (SkippedDetails.FramesTillNextReplication > 0)
-		{
-			SkippedStr = FString::Printf(TEXT("Not ready (%d frames left)"), SkippedDetails.FramesTillNextReplication);
-		}
-		else
+		Ar.Logf(TEXT("[%d Skipped Actors]"), PrioritizedList->SkippedDebugDetails->Num());
+
+		FNativeClassAccumulator DormantClasses;
+		FNativeClassAccumulator CulledClasses;
+
+		for (const FSkippedActorFullDebugDetails& SkippedDetails : *PrioritizedList->SkippedDebugDetails)
 		{
-			SkippedStr = TEXT("Unknown???");
-		}
+			FString SkippedStr;
+			if (SkippedDetails.bWasDormant)
+			{
+				SkippedStr = TEXT("Dormant");
+				DormantClasses.Increment(SkippedDetails.Actor->GetClass());
+			}
+			else if (SkippedDetails.DistanceCulled > 0.f)
+			{
+				SkippedStr = FString::Printf(TEXT("Dist Culled %.2f"), SkippedDetails.DistanceCulled);
+				CulledClasses.Increment(SkippedDetails.Actor->GetClass());
+			}
+			else if (SkippedDetails.FramesTillNextReplication > 0)
+			{
+				SkippedStr = FString::Printf(TEXT("Not ready (%d frames left)"), SkippedDetails.FramesTillNextReplication);
+			}
+			else
+			{
+				SkippedStr = TEXT("Unknown???");
+			}
 
-		Ar.Logf(TEXT("%-40s %s"), *GetActorRepListTypeDebugString(SkippedDetails.Actor), *SkippedStr);
-	}
+			Ar.Logf(TEXT("%-40s %s"), *GetActorRepListTypeDebugString(SkippedDetails.Actor), *SkippedStr);
+		}
 
 		Ar.Logf(TEXT(" Dormant Classes: %s"), *DormantClasses.BuildString());
 		Ar.Logf(TEXT(" Culled Classes: %s"), *CulledClasses.BuildString());
+	}
 
 #endif
 
