Replication graph changes to allow for subconnection calculations as well, matching functionality of the netdriver system.

diff --git a/Engine/Plugins/Runtime/ReplicationGraph/Source/Public/ReplicationGraph.h b/Engine/Plugins/Runtime/ReplicationGraph/Source/Public/ReplicationGraph.h
--- a/Engine/Plugins/Runtime/ReplicationGraph/Source/Public/ReplicationGraph.h
+++ b/Engine/Plugins/Runtime/ReplicationGraph/Source/Public/ReplicationGraph.h
@@ -368,6 +368,7 @@ protected:
 	virtual void GatherActors_DistanceOnly(const FActorRepListRefView& RepList, FGlobalActorReplicationInfoMap& GlobalMap, FPerConnectionActorInfoMap& ConnectionMap, const FConnectionGatherActorListParameters& Params);
 
 	void CalcFrequencyForActor(AActor* Actor, UReplicationGraph* RepGraph, UNetConnection* NetConnection, FGlobalActorReplicationInfo& GlobalInfo, FConnectionReplicationActorInfo& ConnectionInfo, FSettings& MySettings, const FNetViewerArray& Viewers, const uint32 FrameNum, int32 ExistingItemIndex);
+	UE_DEPRECATED(4.23, "Use the other function to allow for multiple viewers")
 	void CalcFrequencyForActor(AActor* Actor, UReplicationGraph* RepGraph, UNetConnection* NetConnection, FGlobalActorReplicationInfo& GlobalInfo, FConnectionReplicationActorInfo& ConnectionInfo, FSettings& MySettings, const FVector& ConnectionViewLocation, const FVector& ConnectionViewDir, const uint32 FrameNum, int32 ExistingItemIndex);
 };
 
@@ -711,9 +712,11 @@ public:
 	/** List of previously (or currently if nothing changed last tick) focused actor data per connection */
 	TArray<FAlwaysRelevantActorInfo, FReplicationGraphConnectionsAllocator> PastRelevantActors;
 
+	UE_DEPRECATED(4.23, "ViewTargets are now handled inside the PastRelevantActorMap")
 	UPROPERTY()
 	AActor* LastViewer = nullptr;
 	
+	UE_DEPRECATED(4.23, "ViewTargets are now handled inside the PastRelevantActorMap")
 	UPROPERTY()
 	AActor* LastViewTarget = nullptr;
 };
@@ -945,10 +948,12 @@ protected:
 
 	/** Default Replication Path */
 	void ReplicateActorListsForConnections_Default(UNetReplicationGraphConnection* ConnectionManager, FGatheredReplicationActorLists& GatheredReplicationListsForConnection, FNetViewerArray& Viewers);
+	UE_DEPRECATED(4.23, "Use the array format to support subconnections as well")
 	void ReplicateActorListsForConnection_Default(UNetReplicationGraphConnection* ConnectionManager, FGatheredReplicationActorLists& GatheredReplicationListsForConnection, FNetViewer& Viewer);
 
 	/** "FastShared" Replication Path */
 	void ReplicateActorListsForConnections_FastShared(UNetReplicationGraphConnection* ConnectionManager, FGatheredReplicationActorLists& GatheredReplicationListsForConnection, FNetViewerArray& Viewers);
+	UE_DEPRECATED(4.23, "Use the array format to support subconnections as well")
 	void ReplicateActorListsForConnection_FastShared(UNetReplicationGraphConnection* ConnectionManager, FGatheredReplicationActorLists& GatheredReplicationListsForConnection, FNetViewer& Viewer);
 
 	/** Connections needing a FlushNet in PostTickDispatch */
@@ -1033,6 +1038,7 @@ public:
 	// ID that is assigned by the replication graph. Will be reassigned/compacted as clients disconnect. Useful for spacing out connection operations. E.g., not stable but always compact.
 	int32 ConnectionId; 
 
+	UE_DEPRECATED(4.23, "Use the LastGatherLocations to have support for subconnection lookups")
 	FVector LastGatherLocation;
 
 	TArray<FLastLocationGatherInfo, FReplicationGraphConnectionsAllocator> LastGatherLocations;
diff --git a/Engine/Plugins/Runtime/ReplicationGraph/Source/Public/ReplicationGraphTypes.h b/Engine/Plugins/Runtime/ReplicationGraph/Source/Public/ReplicationGraphTypes.h
--- a/Engine/Plugins/Runtime/ReplicationGraph/Source/Public/ReplicationGraphTypes.h
+++ b/Engine/Plugins/Runtime/ReplicationGraph/Source/Public/ReplicationGraphTypes.h
@@ -1146,6 +1146,7 @@ struct FConnectionGatherActorListParameters
 
 
 	/** In: The Data the nodes have to work with */
+	UE_DEPRECATED(4.23, "Use the viewer arrays for support for subconnections")
 	FNetViewer& Viewer;
 
 	FNetViewerArray Viewers;
