Fix comment on TearOff RepGraph Node.

diff --git a/Engine/Plugins/Runtime/ReplicationGraph/Source/Public/ReplicationGraph.h b/Engine/Plugins/Runtime/ReplicationGraph/Source/Public/ReplicationGraph.h
--- a/Engine/Plugins/Runtime/ReplicationGraph/Source/Public/ReplicationGraph.h
+++ b/Engine/Plugins/Runtime/ReplicationGraph/Source/Public/ReplicationGraph.h
@@ -740,7 +740,7 @@ struct FTearOffActorInfo
 	bool bHasReppedOnce;
 };
 
-/** Adds actors that are always relevant for a connection. This engine version just adds the PlayerController and ViewTarget (usually the pawn) */
+/** Manages actors that are Tear Off. We will try to replicate these actors one last time to each connection. */
 UCLASS()
 class REPLICATIONGRAPH_API UReplicationGraphNode_TearOff_ForConnection : public UReplicationGraphNode
 {
