RepGraph will now reflush DORM_Initial and non-netstartup actors when a client leaves and rejoins an island.

(Different fix for reverted cl 7461839)

diff --git a/Engine/Plugins/Runtime/ReplicationGraph/Source/Private/ReplicationGraph.cpp b/Engine/Plugins/Runtime/ReplicationGraph/Source/Private/ReplicationGraph.cpp
--- a/Engine/Plugins/Runtime/ReplicationGraph/Source/Private/ReplicationGraph.cpp
+++ b/Engine/Plugins/Runtime/ReplicationGraph/Source/Private/ReplicationGraph.cpp
@@ -2462,7 +2462,7 @@ void UNetReplicationGraphConnection::NotifyClientVisibleLevelNamesAdd(FName Leve
 		TArray<AActor*>& Actors = StreamingWorld->PersistentLevel->Actors;
 		for (AActor* Actor : Actors)
 		{
-			if (Actor && Actor->GetIsReplicated() && (Actor->NetDormancy == DORM_DormantAll))
+			if (Actor->NetDormancy == DORM_DormantAll || (Actor->NetDormancy == DORM_Initial && Actor->IsNetStartupActor() == false))
 			{
 				if (FConnectionReplicationActorInfo* ActorInfo = ActorInfoMap.Find(Actor))
 				{
