Fix nullptr crash in RepGraph

#rnx

diff --git a/Engine/Plugins/Runtime/ReplicationGraph/Source/Private/ReplicationGraph.cpp b/Engine/Plugins/Runtime/ReplicationGraph/Source/Private/ReplicationGraph.cpp
--- a/Engine/Plugins/Runtime/ReplicationGraph/Source/Private/ReplicationGraph.cpp
+++ b/Engine/Plugins/Runtime/ReplicationGraph/Source/Private/ReplicationGraph.cpp
@@ -2462,7 +2462,7 @@ void UNetReplicationGraphConnection::NotifyClientVisibleLevelNamesAdd(FName Leve
 		TArray<AActor*>& Actors = StreamingWorld->PersistentLevel->Actors;
 		for (AActor* Actor : Actors)
 		{
-			if (Actor->NetDormancy == DORM_DormantAll || (Actor->NetDormancy == DORM_Initial && Actor->IsNetStartupActor() == false))
+			if (Actor && (Actor->NetDormancy == DORM_DormantAll || (Actor->NetDormancy == DORM_Initial && Actor->IsNetStartupActor() == false)))
 			{
 				if (FConnectionReplicationActorInfo* ActorInfo = ActorInfoMap.Find(Actor))
 				{
