Fix a bug where changing an Actor's Dormancy could add it back to the GridSpatialization2D node in RepGraph, even if it was explicitly removed.

diff --git a/Engine/Plugins/Runtime/ReplicationGraph/Source/Private/ReplicationGraph.cpp b/Engine/Plugins/Runtime/ReplicationGraph/Source/Private/ReplicationGraph.cpp
--- a/Engine/Plugins/Runtime/ReplicationGraph/Source/Private/ReplicationGraph.cpp
+++ b/Engine/Plugins/Runtime/ReplicationGraph/Source/Private/ReplicationGraph.cpp
@@ -4137,6 +4137,12 @@ void UReplicationGraphNode_GridSpatialization2D::RemoveActor_Dormancy(const FNew
 		{
 			RemoveActorInternal_Dynamic(ActorInfo);
 		}
+
+		// AddActorInternal_Static and AddActorInternal_Dynamic will both override Actor information if they are called repeatedly.
+		// This means that even if AddActor_Dormancy is called multiple times with the same Actor, a single call to RemoveActor_Dormancy
+		// will completely remove the Actor from either the Static or Dynamic list appropriately.
+		// Therefore, it should be safe to call RemoveAll and not worry about trying to track individual delegate handles.
+		ActorRepInfo.Events.DormancyChange.RemoveAll(this);
 	}
 }
 
