Test for nullptr in FActorRepListRefView::ConditionalAdd

diff --git a/Engine/Plugins/Runtime/ReplicationGraph/Source/Public/ReplicationGraphTypes.h b/Engine/Plugins/Runtime/ReplicationGraph/Source/Public/ReplicationGraphTypes.h
--- a/Engine/Plugins/Runtime/ReplicationGraph/Source/Public/ReplicationGraphTypes.h
+++ b/Engine/Plugins/Runtime/ReplicationGraph/Source/Public/ReplicationGraphTypes.h
@@ -79,6 +79,11 @@ FORCEINLINE bool IsActorValidForReplication(const FActorRepListType& In) { retur
 // Tests if an actor is valid for replication gathering. Meaning, it can be gathered from the replication graph and considered for replication.
 FORCEINLINE bool IsActorValidForReplicationGather(const FActorRepListType& In)
 { 
+	if (In == nullptr)
+	{
+		return false;
+	}
+
 	if (!IsActorValidForReplication(In))
 		return false;
 
