- Net.RepGraph.PrintCullDistances will now work without a RepGraphDebugActor
- Added Net.RepGraph.PrintCullDistancesForConnection which outputs also the connection specific cull distance (uses the  RepGraphDebugActor)
- Added DO_ENABLE_REPGRAPH_DEBUG_ACTOR instead of hardcoded build defines to easily add the RepGraphDebugActor

diff --git a/Engine/Plugins/Runtime/ReplicationGraph/Source/Private/ReplicationGraph.cpp b/Engine/Plugins/Runtime/ReplicationGraph/Source/Private/ReplicationGraph.cpp
--- a/Engine/Plugins/Runtime/ReplicationGraph/Source/Private/ReplicationGraph.cpp
+++ b/Engine/Plugins/Runtime/ReplicationGraph/Source/Private/ReplicationGraph.cpp
@@ -1025,7 +1025,7 @@ int32 UReplicationGraph::ServerReplicateActors(float DeltaSeconds)
 				ConnectionManager->ReplicateDormantDestructionInfos();
 			}
 
-#if !(UE_BUILD_SHIPPING || UE_BUILD_TEST)
+#if DO_ENABLE_REPGRAPH_DEBUG_ACTOR
 			{
 				RG_QUICK_SCOPE_CYCLE_COUNTER(NET_ReplicateActors_ReplicateDebugActor);
 				if (ConnectionManager->DebugActor)
@@ -2209,7 +2209,7 @@ void UNetReplicationGraphConnection::Serialize(FArchive& Ar)
 
 void UNetReplicationGraphConnection::TearDown()
 {
-#if !(UE_BUILD_SHIPPING || UE_BUILD_TEST)
+#if DO_ENABLE_REPGRAPH_DEBUG_ACTOR
 	if (DebugActor)
 	{
 		DebugActor->Destroy();
@@ -2308,7 +2308,7 @@ void UNetReplicationGraphConnection::InitForConnection(UNetConnection* InConnect
 	NetConnection = InConnection;
 	InConnection->SetReplicationConnectionDriver(this);
 
-#if !(UE_BUILD_SHIPPING || UE_BUILD_TEST)
+#if DO_ENABLE_REPGRAPH_DEBUG_ACTOR
 	UReplicationGraph* Graph = Cast<UReplicationGraph>(GetOuter());
 	DebugActor = Graph->CreateDebugActor();
 	if (DebugActor)
diff --git a/Engine/Plugins/Runtime/ReplicationGraph/Source/Private/ReplicationGraphDebugging.cpp b/Engine/Plugins/Runtime/ReplicationGraph/Source/Private/ReplicationGraphDebugging.cpp
--- a/Engine/Plugins/Runtime/ReplicationGraph/Source/Private/ReplicationGraphDebugging.cpp
+++ b/Engine/Plugins/Runtime/ReplicationGraph/Source/Private/ReplicationGraphDebugging.cpp
@@ -66,11 +66,6 @@ UNetConnection* AReplicationGraphDebugActor::GetNetConnection() const
 
 // -------------------------------------------------------------
 
-bool AReplicationGraphDebugActor::ServerStartDebugging_Validate()
-{
-	return true;
-}
-
 void AReplicationGraphDebugActor::ServerStartDebugging_Implementation()
 {
 	UE_LOG(LogReplicationGraph, Display, TEXT("ServerStartDebugging"));
@@ -148,44 +143,53 @@ FAutoConsoleCommandWithWorldAndArgs NetRepGraphDebugActorStart(TEXT("Net.RepGrap
 
 // -------------------------------------------------------------
 
-bool AReplicationGraphDebugActor::ServerStopDebugging_Validate()
-{
-	return true;
-}
-
 void AReplicationGraphDebugActor::ServerStopDebugging_Implementation()
 {
 	
 }
 
 // -------------------------------------------------------------
-
-void AReplicationGraphDebugActor::PrintCullDistances()
+void UReplicationGraph::DebugPrintCullDistances(UNetReplicationGraphConnection* SpecificConnection) const
 {
 	struct FData
 	{
 		UClass* Class = nullptr;
-		float DistSq;
+		float Dist;
+		float ConnectionDist;
 		int32 Count;
 	};
 
 	TArray<FData> DataList;
 
-	for (auto It = ReplicationGraph->GlobalActorReplicationInfoMap.CreateActorMapIterator(); It; ++It)
+	for (auto It = GlobalActorReplicationInfoMap.CreateActorMapIterator(); It; ++It)
 	{
 		AActor* Actor = It.Key();
-		TUniquePtr<FGlobalActorReplicationInfo>& InfoPtr = It.Value();
+		const TUniquePtr<FGlobalActorReplicationInfo>& InfoPtr = It.Value();
 		if (!InfoPtr || InfoPtr.Get() == nullptr)
 		{
 			continue;
 		}
 
-		FGlobalActorReplicationInfo& Info = *InfoPtr.Get();
+		const FGlobalActorReplicationInfo& Info = *InfoPtr.Get();
+
+		float ConnectionCullDist = 0.0f;
+
+		// Optionally find this connection's CullDistance
+		if (SpecificConnection)
+		{
+			FPerConnectionActorInfoMap& ConnectionInfo = SpecificConnection->ActorInfoMap;
+			if (FConnectionReplicationActorInfo* ConnectionActorInfo = ConnectionInfo.Find(Actor))
+			{
+				ConnectionCullDist = ConnectionActorInfo->GetCullDistance();
+			}
+		}
 
 		bool bFound = false;
 		for (FData& ExistingData : DataList)
 		{
-			if (ExistingData.Class == Actor->GetClass() && FMath::IsNearlyZero(ExistingData.DistSq - Info.Settings.GetCullDistanceSquared()))
+			if (ExistingData.Class == Actor->GetClass() && 
+				FMath::IsNearlyEqual(ExistingData.Dist, Info.Settings.GetCullDistance()) &&
+				FMath::IsNearlyEqual(ExistingData.ConnectionDist, ConnectionCullDist))
 			{
 				ExistingData.Count++;
 				bFound = true;
@@ -200,32 +204,65 @@ void AReplicationGraphDebugActor::PrintCullDistances()
 
 		FData NewData;
 		NewData.Class = Actor->GetClass();
-		NewData.DistSq = Info.Settings.GetCullDistanceSquared();
+		NewData.Dist = Info.Settings.GetCullDistance();
+		NewData.ConnectionDist = ConnectionCullDist;
 		NewData.Count = 1;
 		DataList.Add(NewData);
 	}
 
-	DataList.Sort([](const FData& LHS, const FData& RHS) { return LHS.DistSq < RHS.DistSq; });
+	DataList.Sort([](const FData& LHS, const FData& RHS) { return LHS.Dist < RHS.Dist; });
 
 	for (FData& Data : DataList)
 	{
 		const UClass* NativeParent = Data.Class;
-		while(NativeParent && !NativeParent->IsNative())
+		while (NativeParent && !NativeParent->IsNative())
 		{
 			NativeParent = NativeParent->GetSuperClass();
 		}
 
-
-		UE_LOG(LogReplicationGraph, Display, TEXT("%s (%s) [%d] = %.2f"), *GetNameSafe(Data.Class), *GetNameSafe(NativeParent), Data.Count, FMath::Sqrt(Data.DistSq));
+		if( SpecificConnection == nullptr )
+		{
+			UE_LOG(LogReplicationGraph, Display, TEXT("%s (%s) [%d] GlobalCullDistance (%.2f)"), *GetNameSafe(Data.Class), *GetNameSafe(NativeParent), Data.Count, Data.Dist);
+		}
+		else
+		{
+			UE_LOG(LogReplicationGraph, Display, TEXT("%s (%s) [%d] GlobalCullDistance (%.2f) ConnectionCullDistance (%.2f)"), *GetNameSafe(Data.Class), *GetNameSafe(NativeParent), Data.Count, Data.Dist, Data.ConnectionDist);
+		}
 	}
 }
 
-FAutoConsoleCommandWithWorldAndArgs NetRepGraphPrintCullDistancesCommand(TEXT("Net.RepGraph.PrintCullDistances"),TEXT(""),
+void AReplicationGraphDebugActor::ServerPrintCullDistances_Implementation()
+{
+	PrintCullDistances();
+}
+
+void AReplicationGraphDebugActor::PrintCullDistances()
+{
+	ReplicationGraph->DebugPrintCullDistances(ConnectionManager);
+}
+
+// -------------------------------------------------------------
+
+FAutoConsoleCommandWithWorldAndArgs NetRepGraphPrintCullDistancesForConnectionCommand(TEXT("Net.RepGraph.PrintCullDistancesForConnection"),TEXT("Print the cull distances via the ReplicationDebugActor to add the connection cull distances."),
 	FConsoleCommandWithWorldAndArgsDelegate::CreateLambda([](const TArray<FString>& Args, UWorld* World)
 	{
 		for (TActorIterator<AReplicationGraphDebugActor> It(World); It; ++It)
 		{
-			It->PrintCullDistances();
+			It->ServerPrintCullDistances();
+		}
+	})
+);
+
+FAutoConsoleCommandWithWorldAndArgs NetRepGraphPrintCullDistancesCommand(TEXT("Net.RepGraph.PrintCullDistances"), TEXT("Print the cull distances of RepGraph actors on the server."),
+	FConsoleCommandWithWorldAndArgsDelegate::CreateLambda([](const TArray<FString>& Args, UWorld* World)
+	{
+		for (TObjectIterator<UReplicationGraph> It; It; ++It)
+		{
+			if (It->NetDriver && It->NetDriver->GetNetMode() != NM_Client)
+			{
+				It->DebugPrintCullDistances();
+				break;
+			}
 		}
 	})
 );
@@ -238,11 +275,6 @@ void AReplicationGraphDebugActor::ServerPrintAllActorInfo_Implementation(const F
 	PrintAllActorInfo(Str);
 }
 
-bool AReplicationGraphDebugActor::ServerPrintAllActorInfo_Validate(const FString& Str)
-{
-	return true;
-}
-
 void AReplicationGraphDebugActor::PrintAllActorInfo(FString MatchString)
 {
 	auto Matches = [&MatchString](UObject* Obj) { return MatchString.IsEmpty() || Obj->GetPathName().Contains(MatchString); };
@@ -326,11 +358,6 @@ FAutoConsoleCommandWithWorldAndArgs NetRepGraphPrintAllActorInfoCmd(TEXT("Net.Re
 
 // -------------------------------------------------------------
 
-bool AReplicationGraphDebugActor::ServerCellInfo_Validate()
-{
-	return true;
-}
-
 void AReplicationGraphDebugActor::ServerCellInfo_Implementation()
 {
 	TArray<FVector, FReplicationGraphConnectionsAllocator> LocationsSent;
@@ -426,11 +453,6 @@ FAutoConsoleCommandWithWorldAndArgs NetRepGraphCellInfo(TEXT("Net.RepGraph.Spati
 
 // -------------------------------------------------------------
 
-bool AReplicationGraphDebugActor::ServerSetCullDistanceForClass_Validate(UClass* Class, float CullDistance)
-{
-	return true;
-}
-
 void AReplicationGraphDebugActor::ServerSetCullDistanceForClass_Implementation(UClass* Class, float CullDistance)
 {
 	if (!Class)
@@ -494,11 +516,6 @@ FAutoConsoleCommandWithWorldAndArgs NetRepGraphSetClassCullDistance(TEXT("Net.Re
 
 // -------------------------------------------------------------
 
-bool AReplicationGraphDebugActor::ServerSetPeriodFrameForClass_Validate(UClass* Class, int32 PeriodFrame)
-{
-	return true;
-}
-
 void AReplicationGraphDebugActor::ServerSetPeriodFrameForClass_Implementation(UClass* Class, int32 PeriodFrame)
 {
 	if (!Class)
@@ -560,11 +577,6 @@ FAutoConsoleCommandWithWorldAndArgs NetRepGraphSetPeriodFrame(TEXT("Net.RepGraph
 
 // -------------------------------------------------------------
 
-bool AReplicationGraphDebugActor::ServerSetConditionalActorBreakpoint_Validate(AActor* Actor)
-{
-	return true;
-}
-
 void AReplicationGraphDebugActor::ServerSetConditionalActorBreakpoint_Implementation(AActor* Actor)
 {
 	DebugActorConnectionPair.Actor = Actor;
diff --git a/Engine/Plugins/Runtime/ReplicationGraph/Source/Public/ReplicationGraph.h b/Engine/Plugins/Runtime/ReplicationGraph/Source/Public/ReplicationGraph.h
--- a/Engine/Plugins/Runtime/ReplicationGraph/Source/Public/ReplicationGraph.h
+++ b/Engine/Plugins/Runtime/ReplicationGraph/Source/Public/ReplicationGraph.h
@@ -63,6 +63,8 @@ UReplicationGraph::InitConnectionGraphNodes
 
 #include "ReplicationGraph.generated.h"
 
+#define DO_ENABLE_REPGRAPH_DEBUG_ACTOR !(UE_BUILD_SHIPPING || UE_BUILD_TEST)
+
 UCLASS(transient, config=Engine)
 class REPLICATIONGRAPH_API UReplicationGraphNode : public UObject
 {
@@ -872,10 +874,12 @@ public:
 
 	uint32 GetReplicationGraphFrame() const { return ReplicationGraphFrame; }
 
-#if !(UE_BUILD_SHIPPING || UE_BUILD_TEST)
+#if DO_ENABLE_REPGRAPH_DEBUG_ACTOR
 	virtual class AReplicationGraphDebugActor* CreateDebugActor() const;
 #endif
 
+	/** Print the Net Culling Distance of all actors in the RepGraph. When a connection is specified will also output the cull distance in the connection info list*/
+	void DebugPrintCullDistances(UNetReplicationGraphConnection* SpecificConnection=nullptr) const;
 
 	/** Prioritization Constants: these affect how the final priority of an actor is calculated in the prioritize phase */
 	struct FPrioritizationConstants
@@ -1184,27 +1188,30 @@ public:
 	UPROPERTY()
 	UReplicationGraph* ReplicationGraph;
 
-	UFUNCTION(Server, Reliable, WithValidation)
+	UFUNCTION(Server, Reliable)
 	void ServerStartDebugging();
 
-	UFUNCTION(Server, Reliable, WithValidation)
+	UFUNCTION(Server, Reliable)
 	void ServerStopDebugging();
 
-	UFUNCTION(Server, Reliable, WithValidation)
+	UFUNCTION(Server, Reliable)
 	void ServerCellInfo();
 
-	UFUNCTION(Server, Reliable, WithValidation)
+	UFUNCTION(Server, Reliable)
 	void ServerPrintAllActorInfo(const FString& Str);
 
-	UFUNCTION(Server, Reliable, WithValidation)
+	UFUNCTION(Server, Reliable)
 	void ServerSetCullDistanceForClass(UClass* Class, float CullDistance);
 
-	UFUNCTION(Server, Reliable, WithValidation)
+	UFUNCTION(Server, Reliable)
 	void ServerSetPeriodFrameForClass(UClass* Class, int32 PeriodFrame);
 
-	UFUNCTION(Server, Reliable, WithValidation)
+	UFUNCTION(Server, Reliable)
 	void ServerSetConditionalActorBreakpoint(AActor* Actor);
 
+	UFUNCTION(Server, Reliable)
+	void ServerPrintCullDistances();
+
 	UFUNCTION(Client, Reliable)
 	void ClientCellInfo(FVector CellLocation, FVector CellExtent, const TArray<AActor*>& Actors);
 
diff --git a/Engine/Plugins/Runtime/ReplicationGraph/Source/Public/ReplicationGraphTypes.h b/Engine/Plugins/Runtime/ReplicationGraph/Source/Public/ReplicationGraphTypes.h
--- a/Engine/Plugins/Runtime/ReplicationGraph/Source/Public/ReplicationGraphTypes.h
+++ b/Engine/Plugins/Runtime/ReplicationGraph/Source/Public/ReplicationGraphTypes.h
@@ -720,6 +720,7 @@ struct FGlobalActorReplicationInfoMap
 	FORCEINLINE void SetClassInfo(UClass* InClass, const FClassReplicationInfo& Info) {	ClassMap.Set(InClass, Info); }
 	
 	FORCEINLINE TMap<FActorRepListType, TUniquePtr<FGlobalActorReplicationInfo>>::TIterator CreateActorMapIterator() { return ActorMap.CreateIterator(); }
+	FORCEINLINE TMap<FActorRepListType, TUniquePtr<FGlobalActorReplicationInfo>>::TConstIterator CreateActorMapIterator() const { return ActorMap.CreateConstIterator(); }
 	FORCEINLINE TMap<FObjectKey, FClassReplicationInfo>::TIterator CreateClassMapIterator() { return ClassMap.CreateIterator(); }
 
 	int32 Num() const { return ActorMap.Num(); }
