Limit connection view locations to HALF_WORLD bounds in the ReplicationGraph

diff --git a/Engine/Plugins/Runtime/ReplicationGraph/Source/Private/ReplicationGraph.cpp b/Engine/Plugins/Runtime/ReplicationGraph/Source/Private/ReplicationGraph.cpp
--- a/Engine/Plugins/Runtime/ReplicationGraph/Source/Private/ReplicationGraph.cpp
+++ b/Engine/Plugins/Runtime/ReplicationGraph/Source/Private/ReplicationGraph.cpp
@@ -4902,19 +4902,24 @@ void UReplicationGraphNode_GridSpatialization2D::GatherActorListsForConnection(c
 		{
 			ClampedViewLoc = GridBounds.GetClosestPointTo(ClampedViewLoc);
 		}
+		else
+		{
+			// Prevent extreme locations from causing the Grid to grow too large
+			ClampedViewLoc = ClampedViewLoc.BoundToCube(HALF_WORLD_MAX);
+		}
 
 		// Find out what bucket the view is in
 		int32 CellX = (ClampedViewLoc.X - SpatialBias.X) / CellSize;
 		if (CellX < 0)
 		{
-			UE_LOG(LogReplicationGraph, Log, TEXT("Net view location.X %s is less than the spatial bias %s"), *ClampedViewLoc.ToString(), *SpatialBias.ToString());
+			UE_LOG(LogReplicationGraph, Log, TEXT("Net view location.X %s is less than the spatial bias %s for %s"), *ClampedViewLoc.ToString(), *SpatialBias.ToString(), *CurViewer.Connection->Describe());
 			CellX = 0;
 		}
 
 		int32 CellY = (ClampedViewLoc.Y - SpatialBias.Y) / CellSize;
 		if (CellY < 0)
 		{
-			UE_LOG(LogReplicationGraph, Log, TEXT("Net view location.Y %s is less than the spatial bias %s"), *ClampedViewLoc.ToString(), *SpatialBias.ToString());
+			UE_LOG(LogReplicationGraph, Log, TEXT("Net view location.Y %s is less than the spatial bias %s for %s"), *ClampedViewLoc.ToString(), *SpatialBias.ToString(), *CurViewer.Connection->Describe());
 			CellY = 0;
 		}
 
@@ -4928,12 +4933,17 @@ void UReplicationGraphNode_GridSpatialization2D::GatherActorListsForConnection(c
 			GatherInfoForConnection = &LastLocationArray[LastLocationArray.Add(FLastLocationGatherInfo(CurViewer.Connection, FVector(ForceInitToZero)))];
 		}
 
-		// Clean up the location data for this connection to be grid bound
 		FVector LastLocationForConnection = GatherInfoForConnection->LastLocation;
 		if (GridBounds.IsValid)
 		{
+			// Clean up the location data for this connection to be grid bound
 			LastLocationForConnection = GridBounds.GetClosestPointTo(LastLocationForConnection);
 		}
+		else
+		{
+			// Prevent extreme locations from causing the Grid to grow too large
+			LastLocationForConnection = LastLocationForConnection.BoundToCube(HALF_WORLD_MAX);
+		}
 
 		// Try to determine the previous location of the user.
 		NewPlayerCell.PrevLocation.X = FMath::Max(0, (int32)((LastLocationForConnection.X - SpatialBias.X) / CellSize));
