Added log when RepGraph discovery budget is disabled

diff --git a/Engine/Plugins/Runtime/ReplicationGraph/Source/Private/ReplicationGraph.cpp b/Engine/Plugins/Runtime/ReplicationGraph/Source/Private/ReplicationGraph.cpp
--- a/Engine/Plugins/Runtime/ReplicationGraph/Source/Private/ReplicationGraph.cpp
+++ b/Engine/Plugins/Runtime/ReplicationGraph/Source/Private/ReplicationGraph.cpp
@@ -2153,6 +2153,7 @@ void UReplicationGraph::SetActorDiscoveryBudget(int32 ActorDiscoveryBudgetInKByt
 	if (ActorDiscoveryBudgetInKBytesPerSec <= 0)
 	{
 		ActorDiscoveryMaxBitsPerFrame = 0;
+		UE_LOG(LogReplicationGraph, Display, TEXT("SetActorDiscoveryBudget disabled the ActorDiscovery budget."));
 		return;
 	}
 
