Added perf stats  in FlushDormancy handlers

#rnx

diff --git a/Engine/Plugins/Runtime/ReplicationGraph/Source/Private/ReplicationGraph.cpp b/Engine/Plugins/Runtime/ReplicationGraph/Source/Private/ReplicationGraph.cpp
--- a/Engine/Plugins/Runtime/ReplicationGraph/Source/Private/ReplicationGraph.cpp
+++ b/Engine/Plugins/Runtime/ReplicationGraph/Source/Private/ReplicationGraph.cpp
@@ -595,7 +595,7 @@ void UReplicationGraph::ForceNetUpdate(AActor* Actor)
 
 void UReplicationGraph::FlushNetDormancy(AActor* Actor, bool bWasDormInitial)
 {
-	RG_QUICK_SCOPE_CYCLE_COUNTER(UReplicationGraph_FlushNetDormancy);
+	QUICK_SCOPE_CYCLE_COUNTER(UReplicationGraph_FlushNetDormancy);
 
 	if (Actor->IsActorInitialized() == false)
 	{
diff --git a/Engine/Source/Runtime/Engine/Private/NetworkDriver.cpp b/Engine/Source/Runtime/Engine/Private/NetworkDriver.cpp
--- a/Engine/Source/Runtime/Engine/Private/NetworkDriver.cpp
+++ b/Engine/Source/Runtime/Engine/Private/NetworkDriver.cpp
@@ -3158,6 +3158,7 @@ void UNetDriver::NotifyActorDormancyChange(AActor* Actor, ENetDormancy OldDorman
 
 void UNetDriver::FlushActorDormancyInternal(AActor *Actor)
 {
+	QUICK_SCOPE_CYCLE_COUNTER(STAT_NetDriver_FlushActorDormancy);
 	// Go through each connection and remove the actor from the dormancy list
 	for (int32 i=0; i < ClientConnections.Num(); ++i)
 	{
