Fix UReplicationGraphNode_DynamicSpatialFrequency::CalcFrequencyForActor sometimes not finding valid zones due to unclamped dot product.

#rnx

diff --git a/Engine/Plugins/Runtime/ReplicationGraph/Source/Private/ReplicationGraph.cpp b/Engine/Plugins/Runtime/ReplicationGraph/Source/Private/ReplicationGraph.cpp
--- a/Engine/Plugins/Runtime/ReplicationGraph/Source/Private/ReplicationGraph.cpp
+++ b/Engine/Plugins/Runtime/ReplicationGraph/Source/Private/ReplicationGraph.cpp
@@ -3383,7 +3383,7 @@ FORCEINLINE void UReplicationGraphNode_DynamicSpatialFrequency::CalcFrequencyFor
 	const FVector DirToActor = (GlobalInfo.WorldLocation - LowestDistanceViewer->ViewLocation);
 	const float DistanceToActor = FMath::Sqrt(SmallestDistanceToActorSq);
 	const FVector NormDirToActor = DistanceToActor > SMALL_NUMBER ? (DirToActor / DistanceToActor) : DirToActor;
-	const float DotP = FVector::DotProduct(NormDirToActor, ConnectionViewDir);
+	const float DotP = FMath::Clamp(FVector::DotProduct(NormDirToActor, ConnectionViewDir), -1.0f, 1.0f);
 
 	const bool ActorSupportsFastShared = (GlobalInfo.Settings.FastSharedReplicationFunc != nullptr);
 	TArrayView<FSpatializationZone>& ZoneList = ActorSupportsFastShared ? MySettings.ZoneSettings : MySettings.ZoneSettings_NonFastSharedActors;
