Add role swapping to fast shared path

diff --git a/Engine/Plugins/Runtime/ReplicationGraph/Source/Private/ReplicationGraph.cpp b/Engine/Plugins/Runtime/ReplicationGraph/Source/Private/ReplicationGraph.cpp
--- a/Engine/Plugins/Runtime/ReplicationGraph/Source/Private/ReplicationGraph.cpp
+++ b/Engine/Plugins/Runtime/ReplicationGraph/Source/Private/ReplicationGraph.cpp
@@ -1562,6 +1562,12 @@ int64 UReplicationGraph::ReplicateSingleActor_FastShared(AActor* Actor, FConnect
 		ConnectionData.FastPath_NextReplicationFrameNum = FrameNum + ConnectionData.FastPath_ReplicationPeriodFrame;
 	};
 	
+	TOptional<FScopedActorRoleSwap> SwapGuard;
+	if (GlobalActorInfo.bSwapRolesOnReplicate)
+	{
+		SwapGuard = FScopedActorRoleSwap(Actor);
+	}
+
 	if (CVar_RepGraph_FastShared_ForceFull > 0)
 	{
 		return ReplicateSingleActor(Actor, ConnectionData, GlobalActorInfo, FindOrAddConnectionManager(NetConnection)->ActorInfoMap, ConnectionManager, FrameNum);
