ReplicationGraph
* Upgrade ReplicationPeriodFrame to uint16 to fix overflow error when the NetUpdateFrequency is very low

PR #6330: Fix bug about implict cast of ReplicationPeriodFrame from uint32 to uin8. (Contributed by [REDACTED_PERSON])

diff --git a/Engine/Plugins/Runtime/ReplicationGraph/Source/Private/BasicReplicationGraph.cpp b/Engine/Plugins/Runtime/ReplicationGraph/Source/Private/BasicReplicationGraph.cpp
--- a/Engine/Plugins/Runtime/ReplicationGraph/Source/Private/BasicReplicationGraph.cpp
+++ b/Engine/Plugins/Runtime/ReplicationGraph/Source/Private/BasicReplicationGraph.cpp
@@ -38,7 +38,7 @@ void UBasicReplicationGraph::InitGlobalActorClassSettings()
 		FClassReplicationInfo ClassInfo;
 
 		// Replication Graph is frame based. Convert NetUpdateFrequency to ReplicationPeriodFrame based on Server MaxTickRate.
-		ClassInfo.ReplicationPeriodFrame = FMath::Max<uint32>( (uint32)FMath::RoundToFloat(NetDriver->NetServerMaxTickRate / ActorCDO->NetUpdateFrequency), 1);
+		ClassInfo.ReplicationPeriodFrame = GetReplicationPeriodFrameForFrequency(ActorCDO->NetUpdateFrequency);
 		
 		if (ActorCDO->bAlwaysRelevant || ActorCDO->bOnlyRelevantToOwner)
 		{
diff --git a/Engine/Plugins/Runtime/ReplicationGraph/Source/Private/ReplicationGraph.cpp b/Engine/Plugins/Runtime/ReplicationGraph/Source/Private/ReplicationGraph.cpp
--- a/Engine/Plugins/Runtime/ReplicationGraph/Source/Private/ReplicationGraph.cpp
+++ b/Engine/Plugins/Runtime/ReplicationGraph/Source/Private/ReplicationGraph.cpp
@@ -3306,19 +3306,21 @@ void UReplicationGraphNode_DynamicSpatialFrequency::GatherActorListsForConnectio
 REPGRAPH_DEVCVAR_SHIPCONST(int32, "Net.RepGraph.DynamicSpatialFrequency.Draw", CVar_RepGraph_DynamicSpatialFrequency_Draw, 0, "");
 REPGRAPH_DEVCVAR_SHIPCONST(int32, "Net.RepGraph.DynamicSpatialFrequency.ForceMaxFreq", CVar_RepGraph_DynamicSpatialFrequency_ForceMaxFreq, 0, "Forces DSF to set max frame replication periods on all actors (1 frame rep periods). 1 = default replication. 2 = fast path. 3 = Both (effectively, default)");
 
-FORCEINLINE uint32 CalcDynamicReplicationPeriod(const float FinalPCT, const uint32 MinRepPeriod, const uint32 MaxRepPeriod, uint8& OutReplicationPeriodFrame, uint32& OutNextReplicationFrame, const uint32 LastRepFrameNum, const uint32 FrameNum, bool ForFastPath)
+FORCEINLINE uint32 CalcDynamicReplicationPeriod(const float FinalPCT, const uint32 MinRepPeriod, const uint32 MaxRepPeriod, uint16& OutReplicationPeriodFrame, uint32& OutNextReplicationFrame, const uint32 LastRepFrameNum, const uint32 FrameNum, bool ForFastPath)
 {
 	const float PeriodRange = (float)(MaxRepPeriod - MinRepPeriod);
 	const uint32 ExtraPeriod = (uint32)FMath::CeilToInt(PeriodRange * FinalPCT);
 				
-	uint32 FinalPeriod = MinRepPeriod + ExtraPeriod;
-	OutReplicationPeriodFrame = FinalPeriod;
+	const uint32 FinalPeriod = MinRepPeriod + ExtraPeriod;
+	OutReplicationPeriodFrame = (uint16)FMath::Clamp<uint32>(FinalPeriod, 1, MAX_uint16);
 
-	const uint32 NextRepFrameNum = LastRepFrameNum + FinalPeriod;
+	const uint32 NextRepFrameNum = LastRepFrameNum + OutReplicationPeriodFrame;
 	OutNextReplicationFrame = NextRepFrameNum;
 
 
 #if !(UE_BUILD_SHIPPING)
+	ensureMsgf(OutReplicationPeriodFrame == FinalPeriod, TEXT("Overflow error when FinalPeriod(%u) was assigned to OutReplicationPeriodFrame(%u). RepPeriod values are probably too big"), FinalPeriod, OutReplicationPeriodFrame);
+
 	if (CVar_RepGraph_DynamicSpatialFrequency_ForceMaxFreq > 0)
 	{
 		if ((CVar_RepGraph_DynamicSpatialFrequency_ForceMaxFreq == 1 && ForFastPath == 0) ||
diff --git a/Engine/Plugins/Runtime/ReplicationGraph/Source/Private/ReplicationGraphDebugging.cpp b/Engine/Plugins/Runtime/ReplicationGraph/Source/Private/ReplicationGraphDebugging.cpp
--- a/Engine/Plugins/Runtime/ReplicationGraph/Source/Private/ReplicationGraphDebugging.cpp
+++ b/Engine/Plugins/Runtime/ReplicationGraph/Source/Private/ReplicationGraphDebugging.cpp
@@ -526,7 +526,7 @@ void AReplicationGraphDebugActor::ServerSetPeriodFrameForClass_Implementation(UC
 
 	FClassReplicationInfo& ClassInfo = ReplicationGraph->GlobalActorReplicationInfoMap.GetClassInfo(Class);
 	ClassInfo.ReplicationPeriodFrame = PeriodFrame;
-	UE_LOG(LogReplicationGraph, Display, TEXT("Setting ReplicationPeriodFrame for class %s to %d"), *Class->GetName(), PeriodFrame);
+	UE_LOG(LogReplicationGraph, Display, TEXT("Setting ReplicationPeriodFrame for class %s to %u"), *Class->GetName(), ClassInfo.ReplicationPeriodFrame);
 
 	for (TActorIterator<AActor> ActorIt(GetWorld(), Class); ActorIt; ++ActorIt)
 	{
@@ -534,14 +534,14 @@ void AReplicationGraphDebugActor::ServerSetPeriodFrameForClass_Implementation(UC
 		if (FGlobalActorReplicationInfo* ActorInfo = ReplicationGraph->GlobalActorReplicationInfoMap.Find(Actor))
 		{
 			ActorInfo->Settings.ReplicationPeriodFrame = PeriodFrame;
-			UE_LOG(LogReplicationGraph, Display, TEXT("Setting GlobalActorInfo ReplicationPeriodFrame for %s to %d"), *Actor->GetName(), PeriodFrame);
+			UE_LOG(LogReplicationGraph, Display, TEXT("Setting GlobalActorInfo ReplicationPeriodFrame for %s to %u"), *Actor->GetName(), ActorInfo->Settings.ReplicationPeriodFrame);
 		}
 
 
 		if (FConnectionReplicationActorInfo* ConnectionActorInfo = ConnectionManager->ActorInfoMap.Find(Actor))
 		{
 			ConnectionActorInfo->ReplicationPeriodFrame = PeriodFrame;
-			UE_LOG(LogReplicationGraph, Display, TEXT("Setting Connection ReplicationPeriodFrame for %s to %d"), *Actor->GetName(), PeriodFrame);
+			UE_LOG(LogReplicationGraph, Display, TEXT("Setting Connection ReplicationPeriodFrame for %s to %u"), *Actor->GetName(), ConnectionActorInfo->ReplicationPeriodFrame);
 		}
 	}
 }
diff --git a/Engine/Plugins/Runtime/ReplicationGraph/Source/Public/ReplicationGraph.h b/Engine/Plugins/Runtime/ReplicationGraph/Source/Public/ReplicationGraph.h
--- a/Engine/Plugins/Runtime/ReplicationGraph/Source/Public/ReplicationGraph.h
+++ b/Engine/Plugins/Runtime/ReplicationGraph/Source/Public/ReplicationGraph.h
@@ -911,6 +911,18 @@ public:
 
 	void NotifyConnectionSaturated(class UNetReplicationGraphConnection& Connection);
 
+	uint16 GetReplicationPeriodFrameForFrequency(float NetUpdateFrequency) const
+	{
+		check(NetDriver);
+		check(NetUpdateFrequency != 0.0f);
+
+		// Replication Graph is frame based. Convert NetUpdateFrequency to ReplicationPeriodFrame based on Server MaxTickRate.
+		uint32 FramesBetweenUpdates = (uint32)FMath::RoundToInt(NetDriver->NetServerMaxTickRate / NetUpdateFrequency);
+		FramesBetweenUpdates = FMath::Clamp<uint32>(FramesBetweenUpdates, 1, MAX_uint16);
+
+		return (uint16)FramesBetweenUpdates;
+	}
+
 protected:
 
 	virtual void InitializeForWorld(UWorld* World);
diff --git a/Engine/Plugins/Runtime/ReplicationGraph/Source/Public/ReplicationGraphTypes.h b/Engine/Plugins/Runtime/ReplicationGraph/Source/Public/ReplicationGraphTypes.h
--- a/Engine/Plugins/Runtime/ReplicationGraph/Source/Public/ReplicationGraphTypes.h
+++ b/Engine/Plugins/Runtime/ReplicationGraph/Source/Public/ReplicationGraphTypes.h
@@ -364,9 +364,9 @@ struct FClassReplicationInfo
 	float StarvationPriorityScale = 1.f;
 	float AccumulatedNetPriorityBias = 0.f;
 	
-	uint8 ReplicationPeriodFrame = 1;
-	uint8 FastPath_ReplicationPeriodFrame = 1;
-	uint8 ActorChannelFrameTimeout = 4;
+	uint16 ReplicationPeriodFrame = 1;
+	uint16 FastPath_ReplicationPeriodFrame = 1;
+	uint16 ActorChannelFrameTimeout = 4;
 
 	TFunction<bool(AActor*)> FastSharedReplicationFunc = nullptr;
 	FName FastSharedReplicationFuncName = NAME_None;
@@ -865,12 +865,15 @@ public:
 	/** Default replication */
 	uint32	NextReplicationFrameNum = 0;	/** The next frame we are allowed to replicate on */
 	uint32	LastRepFrameNum = 0;			/** The last frame that this actor replicated on to this connection */
-	uint8	ReplicationPeriodFrame = 1;		/** Min frames that have to pass between subsequent calls to ReplicateActor */
 
 	/** FastPath versions of the above */
 	uint32	FastPath_NextReplicationFrameNum = 0;
 	uint32	FastPath_LastRepFrameNum = 0;
-	uint8	FastPath_ReplicationPeriodFrame = 1;
+
+	/** Min frames that have to pass between subsequent calls to ReplicateActor */
+	uint16	ReplicationPeriodFrame = 1;		
+	uint16	FastPath_ReplicationPeriodFrame = 1;
+	
 
 	/** The frame num that we will close the actor channel. This will get updated/pushed anytime the actor replicates based on FGlobalActorReplicationInfo::ActorChannelFrameTimeout  */
 	uint32 ActorChannelCloseFrameNum = 0;
