Added public access to Visible level tracker info

#rnx

diff --git a/Engine/Plugins/Runtime/ReplicationGraph/Source/Public/ReplicationGraphTypes.h b/Engine/Plugins/Runtime/ReplicationGraph/Source/Public/ReplicationGraphTypes.h
--- a/Engine/Plugins/Runtime/ReplicationGraph/Source/Public/ReplicationGraphTypes.h
+++ b/Engine/Plugins/Runtime/ReplicationGraph/Source/Public/ReplicationGraphTypes.h
@@ -1399,13 +1399,15 @@ struct FReplicationGraphCSVTracker
 		ImplicitClassTracker.Set(BaseActorClass, NewData);
 	}
 
-	void VisibleLevelConnectionAdded(FName LevelName)
+	/** Returns true when the level data is created for the first time */
+	bool VisibleLevelConnectionAdded(FName LevelName)
 	{
 #if REPGRAPH_CSV_TRACKER
 		FVisibleLevelData* LevelData = VisibleLevelConnectionTracker.FindByKey(LevelName);
 		if (LevelData)
 		{
 			LevelData->NbConnections++;
+			return false;
 		}
 		else
 		{
@@ -1413,8 +1415,10 @@ struct FReplicationGraphCSVTracker
 			VisibleLevelData.LevelName = LevelName;
 			VisibleLevelData.NbConnections = 1;
 			VisibleLevelConnectionTracker.Add(VisibleLevelData);
+			return true;
 		}
 #endif
+		return false;
 	}
 
 	void VisibleLevelConnectionRemoved(FName LevelName)
@@ -1453,7 +1457,6 @@ struct FReplicationGraphCSVTracker
 #endif //REPGRAPH_CSV_TRACKER
 	}
 
-
 	void PostReplicateActor(UClass* ActorClass, const double Time, const int64 Bits, const bool bIsActorDiscovery)
 	{
 #if REPGRAPH_CSV_TRACKER
@@ -1604,6 +1607,21 @@ struct FReplicationGraphCSVTracker
 		ExplicitClassTracker_FastPath.CountBytes(Ar);
 	}
 
+public:
+
+	struct FVisibleLevelData
+	{
+		FName LevelName;
+		int32 NbConnections;
+		FName CustomReadableName;
+
+		bool operator==(const FVisibleLevelData& Other) const { return LevelName == Other.LevelName; }
+		bool operator==(const FName& InLevelName) const { return LevelName == InLevelName; }
+	};
+
+	/** */
+	const TArray<FReplicationGraphCSVTracker::FVisibleLevelData>& GetVisibleLevelsData() const { return VisibleLevelConnectionTracker; }
+
 private:
 
 	struct FTrackedData
@@ -1639,16 +1657,6 @@ private:
 		FTrackedData Data;
 	};
 
-	struct FVisibleLevelData
-	{
-		FName LevelName;
-		int32 NbConnections;
-		FName CustomReadableName;
-
-		bool operator==(const FVisibleLevelData& Other) const { return LevelName == Other.LevelName; }
-		bool operator==(const FName& InLevelName) const { return LevelName == InLevelName; }
-	};
-
 	TArray<FTrackerItem, TInlineAllocator<1>> ExplicitClassTracker;
 	TClassMap<TSharedPtr<FTrackedData>> ImplicitClassTracker;
 	TArray<TSharedPtr<FTrackedData>> UniqueImplicitTrackedData;
