Repgraph
* For AlwaysRelevant actors we will now always send the DestructionInfo to connections that knew of this actor.
* Added a flag to FActorDestructionInfo that ignores the distance culling test.

[RN] minor engine/network

diff --git a/Engine/Plugins/Runtime/ReplicationGraph/Source/Private/BasicReplicationGraph.cpp b/Engine/Plugins/Runtime/ReplicationGraph/Source/Private/BasicReplicationGraph.cpp
--- a/Engine/Plugins/Runtime/ReplicationGraph/Source/Private/BasicReplicationGraph.cpp
+++ b/Engine/Plugins/Runtime/ReplicationGraph/Source/Private/BasicReplicationGraph.cpp
@@ -110,6 +110,7 @@ void UBasicReplicationGraph::RouteRemoveNetworkActorToNodes(const FNewReplicated
 	if (ActorInfo.Actor->bAlwaysRelevant)
 	{
 		AlwaysRelevantNode->NotifyRemoveNetworkActor(ActorInfo);
+		SetActorDestructionInfoToIgnoreDistanceCulling(ActorInfo.GetActor());
 	}
 	else if (ActorInfo.Actor->bOnlyRelevantToOwner)
 	{
diff --git a/Engine/Plugins/Runtime/ReplicationGraph/Source/Private/ReplicationGraph.cpp b/Engine/Plugins/Runtime/ReplicationGraph/Source/Private/ReplicationGraph.cpp
--- a/Engine/Plugins/Runtime/ReplicationGraph/Source/Private/ReplicationGraph.cpp
+++ b/Engine/Plugins/Runtime/ReplicationGraph/Source/Private/ReplicationGraph.cpp
@@ -2194,6 +2194,29 @@ void UReplicationGraph::NotifyConnectionSaturated(UNetReplicationGraphConnection
 	++GNumSaturatedConnections;
 }
 
+void UReplicationGraph::SetActorDestructionInfoToIgnoreDistanceCulling(AActor* DestroyedActor)
+{
+	if (!DestroyedActor)
+	{
+		return;
+	}
+
+	check(NetDriver);
+	FNetworkGUID NetGUID = NetDriver->GuidCache->GetNetGUID(DestroyedActor);
+			
+	if (NetGUID.IsDefault())
+	{
+		UE_CLOG(CVar_RepGraph_LogNetDormancyDetails > 0, LogReplicationGraph, Warning, TEXT("SetActorDestructionInfoToIgnoreDistanceCulling ignored for %s. No NetGUID assigned to the actor"), *GetNameSafe(DestroyedActor));
+		return;
+	}
+
+	// See if a destruction info exists for this actor
+	if (const TUniquePtr<FActorDestructionInfo>* DestructionInfoPtr = NetDriver->DestroyedStartupOrDormantActors.Find(NetGUID))
+	{
+		(*DestructionInfoPtr)->bIgnoreDistanceCulling = true;
+	}
+}
+
 // --------------------------------------------------------------------------------------------------------------------------------------------
 // --------------------------------------------------------------------------------------------------------------------------------------------
 // --------------------------------------------------------------------------------------------------------------------------------------------
@@ -2509,23 +2532,24 @@ int64 UNetReplicationGraphConnection::ReplicateDestructionInfos(const FNetViewer
 	{
 		const FCachedDestructInfo& Info = PendingDestructInfoList[idx];
 		FActorDestructionInfo* DestructInfo = Info.DestructionInfo;
-		bool bStillInRelevency = false;
+		bool bSendDestructionInfo = false;
 
-		// Find if anyone is close to this object.
-		for (const FNetViewer& CurViewer : Viewers)
+		if (!DestructInfo->bIgnoreDistanceCulling)
 		{
-			float DistSquared = FMath::Square(Info.CachedPosition.X - CurViewer.ViewLocation.X) + FMath::Square(Info.CachedPosition.Y - CurViewer.ViewLocation.Y);
-
-			// Someone is nearby this object, do not remove it.
-			if (!(DistSquared < DestructInfoMaxDistanceSquared))
+			// Only send destruction info if the viewers are close enough to the destroyed actor
+			for (const FNetViewer& CurViewer : Viewers)
 			{
-				bStillInRelevency = true;
-				break;
+				const float DistSquared = FVector::DistSquared2D(Info.CachedPosition, CurViewer.ViewLocation);
+
+				if (DistSquared < DestructInfoMaxDistanceSquared)
+				{
+					bSendDestructionInfo = true;
+					break;
+				}
 			}
 		}
 
-		// Essentially, if no one can see this object, mark it for destruction.
-		if (!bStillInRelevency)
+		if (bSendDestructionInfo || DestructInfo->bIgnoreDistanceCulling)
 		{
 			UActorChannel* Channel = (UActorChannel*)NetConnection->CreateChannelByName(NAME_Actor, EChannelCreateFlags::OpenedLocally);
 			if (Channel)
diff --git a/Engine/Plugins/Runtime/ReplicationGraph/Source/Public/ReplicationGraph.h b/Engine/Plugins/Runtime/ReplicationGraph/Source/Public/ReplicationGraph.h
--- a/Engine/Plugins/Runtime/ReplicationGraph/Source/Public/ReplicationGraph.h
+++ b/Engine/Plugins/Runtime/ReplicationGraph/Source/Public/ReplicationGraph.h
@@ -911,6 +911,8 @@ public:
 
 	void NotifyConnectionSaturated(class UNetReplicationGraphConnection& Connection);
 
+	void SetActorDestructionInfoToIgnoreDistanceCulling(AActor* DestroyedActor);
+
 	uint16 GetReplicationPeriodFrameForFrequency(float NetUpdateFrequency) const
 	{
 		check(NetDriver);
diff --git a/Engine/Source/Runtime/Engine/Classes/Engine/NetDriver.h b/Engine/Source/Runtime/Engine/Classes/Engine/NetDriver.h
--- a/Engine/Source/Runtime/Engine/Classes/Engine/NetDriver.h
+++ b/Engine/Source/Runtime/Engine/Classes/Engine/NetDriver.h
@@ -570,6 +570,9 @@ public:
 	FName			StreamingLevelName;
 	EChannelCloseReason Reason;
 
+	/** When true the destruction info data will be sent even if the viewers are not close to the actor */
+	bool bIgnoreDistanceCulling = false;
+
 	void CountBytes(FArchive& Ar)
 	{
 		PathName.CountBytes(Ar);
diff --git a/Engine/Source/Runtime/Engine/Classes/Engine/ReplicationDriver.h b/Engine/Source/Runtime/Engine/Classes/Engine/ReplicationDriver.h
--- a/Engine/Source/Runtime/Engine/Classes/Engine/ReplicationDriver.h
+++ b/Engine/Source/Runtime/Engine/Classes/Engine/ReplicationDriver.h
@@ -92,6 +92,9 @@ public:
 
 	virtual void NotifyActorDormancyChange(AActor* Actor, ENetDormancy OldDormancyState) PURE_VIRTUAL(UReplicationDriver::NotifyActorDormancyChange, );
 
+	/** Called when a destruction info is created for an actor. Can be used to override some of the destruction info struct */
+	virtual void NotifyDestructionInfoCreated(AActor* Actor, FActorDestructionInfo& DestructionInfo) PURE_VIRTUAL(UReplicationDriver::NotifyDestructionInfoCreated, );
+
 	virtual void SetRoleSwapOnReplicate(AActor* Actor, bool bSwapRoles) PURE_VIRTUAL(UReplicationDriver::SetRoleSwapOnReplicate, );
 
 	/** Handles an RPC. Returns true if it actually handled it. Returning false will cause the rep driver function to handle it instead */
