Optimized RepGraph ReplicateDestructionInfos by not testing the destroyed actors list every frame
* Now testing if viewers are near destroyed objects only after they travel far enough from the previously tested location
* Set Net.RepGraph.OutOfRangeDistanceCheckRatio to 0 to disable the optim and test the Destroyed Actor List every frame.
* Saves 0.2ms per frame in the worst seen case

[RN] minor engine/network

diff --git a/Engine/Plugins/Runtime/ReplicationGraph/Source/Private/ReplicationGraph.cpp b/Engine/Plugins/Runtime/ReplicationGraph/Source/Private/ReplicationGraph.cpp
--- a/Engine/Plugins/Runtime/ReplicationGraph/Source/Private/ReplicationGraph.cpp
+++ b/Engine/Plugins/Runtime/ReplicationGraph/Source/Private/ReplicationGraph.cpp
@@ -99,6 +99,9 @@ static FAutoConsoleVariableRef CVarRepGraphPrintTrackClassReplication(TEXT("Net.
 int32 CVar_RepGraph_DormantDynamicActorsDestruction = 0;
 static FAutoConsoleVariableRef CVarRepGraphDormantDynamicActorsDestruction(TEXT("Net.RepGraph.DormantDynamicActorsDestruction"), CVar_RepGraph_DormantDynamicActorsDestruction, TEXT(""), ECVF_Default );
 
+float CVar_RepGraph_OutOfRangeDistanceCheckRatio = 0.5f;
+static FAutoConsoleVariableRef CVarRepGraphOutOfRangeDistanceCheckRatio(TEXT("Net.RepGraph.OutOfRangeDistanceCheckRatio"), CVar_RepGraph_OutOfRangeDistanceCheckRatio, TEXT("The ratio of DestructInfoMaxDistance that gives the distance traveled before we reevaluate the out of range destroyed actors list"), ECVF_Default );
+
 static TAutoConsoleVariable<float> CVar_ForceConnectionViewerPriority(TEXT("Net.RepGraph.ForceConnectionViewerPriority"), 1, TEXT("Force the connection's player controller and viewing pawn as topmost priority."));
 
 REPGRAPH_DEVCVAR_SHIPCONST(int32, "Net.RepGraph.LogNetDormancyDetails", CVar_RepGraph_LogNetDormancyDetails, 0, "Logs actors that are removed from the replication graph/nodes.");
@@ -888,6 +891,8 @@ int32 UReplicationGraph::ServerReplicateActors(float DeltaSeconds)
 			bWasConnectionSaturated = false;
 		};
 
+		const FReplicationGraphDestructionSettings DestructionSettings(DestructInfoMaxDistanceSquared, CVar_RepGraph_OutOfRangeDistanceCheckRatio * DestructInfoMaxDistanceSquared);
+
 		FBitWriter& ConnectionSendBuffer = NetConnection->SendBuffer; // unused
 		ConnectionManager->QueuedBitsForActorDiscovery = 0;
 
@@ -914,25 +919,7 @@ int32 UReplicationGraph::ServerReplicateActors(float DeltaSeconds)
 				Node->GatherActorListsForConnection(Parameters);
 			}
 
-			// Update all the LastGatherLocations for this connection.
-			for (const FNetViewer& CurViewer : ConnectionViewers)
-			{
-				FLastLocationGatherInfo* LastInfoForViewer = ConnectionManager->LastGatherLocations.FindByKey<UNetConnection*>(CurViewer.Connection);
-				if (LastInfoForViewer != nullptr)
-				{
-					LastInfoForViewer->LastLocation = CurViewer.ViewLocation;
-				}
-				else
-				{
-					// We need to add this viewer to the last gather locations
-					ConnectionManager->LastGatherLocations.Add(FLastLocationGatherInfo(CurViewer.Connection, CurViewer.ViewLocation));
-				}
-			}
-
-			// Clean up any dead entries in the last gather array
-			ConnectionManager->LastGatherLocations.RemoveAll([&](FLastLocationGatherInfo& CurGatherInfo) {
-				return CurGatherInfo.Connection == nullptr;
-			});
+			ConnectionManager->UpdateGatherLocationsForConnection(ConnectionViewers, DestructionSettings);
 
 			// Do this so we don't break anyone.
 			PRAGMA_DISABLE_DEPRECATION_WARNINGS
@@ -1021,7 +1008,7 @@ int32 UReplicationGraph::ServerReplicateActors(float DeltaSeconds)
 			// ------------------------------------------
 			{
 				QUICK_SCOPE_CYCLE_COUNTER(NET_ReplicateActors_ReplicateDestructionInfos);
-				ConnectionManager->ReplicateDestructionInfos(ConnectionViewers, DestructInfoMaxDistanceSquared);
+				ConnectionManager->ReplicateDestructionInfos(ConnectionViewers, DestructionSettings.DestructInfoMaxDistanceSquared);
 			}
 
 			// ------------------------------------------
@@ -2412,8 +2399,9 @@ void UNetReplicationGraphConnection::NotifyAddDestructionInfo(FActorDestructionI
 		return;
 	}
 
-	// TEMP: Disable this completely in Test Builds as well once we verify crash is fixed.
-#if !(UE_BUILD_SHIPPING) 
+
+
+#if !(UE_BUILD_SHIPPING || UE_BUILD_TEST) 
 	// Should not be happening but lets check in non shipping builds.
 	int32 ExistingIdx = PendingDestructInfoList.IndexOfByKey(DestructInfo);	
 	if (!ensureMsgf(ExistingIdx == INDEX_NONE, TEXT("::NotifyAddDestructionInfo already contains DestructInfo: 0x%X (%s)"), (int64)DestructInfo, *DestructInfo->PathName))
@@ -2447,16 +2435,19 @@ void UNetReplicationGraphConnection::NotifyAddDormantDestructionInfo(AActor* Act
 
 void UNetReplicationGraphConnection::NotifyRemoveDestructionInfo(FActorDestructionInfo* DestructInfo)
 {
-	int32 RemoveIdx = PendingDestructInfoList.IndexOfByKey(DestructInfo);
-	if (RemoveIdx != INDEX_NONE)
+	const FCachedDestructInfo CachedDestructInfo(DestructInfo);
+
+	bool bRemoved = PendingDestructInfoList.RemoveSingleSwap(DestructInfo, false) > 0;
+
+	// Check if the actor is in the out of range list
+	if( !bRemoved )
 	{
-		PendingDestructInfoList.RemoveAtSwap(RemoveIdx, 1, false);
+		OutOfRangeDestroyedActors.RemoveSingleSwap(DestructInfo, false);
 	}
-
+	
 	TrackedDestructionInfoPtrs.Remove(DestructInfo);
 
-	// TEMP: Disable this completely in Test Builds as well once we verify crash is fixed.
-#if !(UE_BUILD_SHIPPING)
+#if !(UE_BUILD_SHIPPING || UE_BUILD_TEST)
 	// Check that its totally gone. Should not be happening!
 	int32 DuplicateIdx = INDEX_NONE;
 	while(true)
@@ -2476,6 +2467,7 @@ void UNetReplicationGraphConnection::NotifyResetDestructionInfo()
 {
 	TrackedDestructionInfoPtrs.Reset();
 	PendingDestructInfoList.Reset();
+	OutOfRangeDestroyedActors.Reset();
 }
 
 void UNetReplicationGraphConnection::GetClientVisibleLevelNames(TSet<FName>& OutLevelNames) const
@@ -2525,12 +2517,12 @@ void UNetReplicationGraphConnection::NotifyClientVisibleLevelNamesAdd(FName Leve
 
 int64 UNetReplicationGraphConnection::ReplicateDestructionInfos(const FNetViewerArray& Viewers, const float DestructInfoMaxDistanceSquared)
 {
-	CSV_SCOPED_TIMING_STAT_EXCLUSIVE(ReplicateDestructionInfos);
+	QUICK_SCOPE_CYCLE_COUNTER(ReplicateDestructionInfos);
 
 	int64 NumBits = 0;
 	for (int32 idx=PendingDestructInfoList.Num()-1; idx >=0; --idx)
 	{
-		const FCachedDestructInfo& Info = PendingDestructInfoList[idx];
+		FCachedDestructInfo& Info = PendingDestructInfoList[idx];
 		FActorDestructionInfo* DestructInfo = Info.DestructionInfo;
 		bool bSendDestructionInfo = false;
 
@@ -2560,6 +2552,12 @@ int64 UNetReplicationGraphConnection::ReplicateDestructionInfos(const FNetViewer
 			PendingDestructInfoList.RemoveAtSwap(idx, 1, false);
 			TrackedDestructionInfoPtrs.Remove(DestructInfo);
 		}
+		else
+		{
+			// Add the far actor to the out of range list so we don't evaluate it every frame
+			OutOfRangeDestroyedActors.Emplace(MoveTemp(Info));
+			PendingDestructInfoList.RemoveAtSwap(idx, 1, false);
+		}
 	}
 
 	return NumBits;
@@ -2594,6 +2592,59 @@ int64 UNetReplicationGraphConnection::ReplicateDormantDestructionInfos()
 	return NumBits;
 }
 
+void UNetReplicationGraphConnection::UpdateGatherLocationsForConnection(const FNetViewerArray& ConnectionViewers, const FReplicationGraphDestructionSettings& DestructionSettings)
+{
+	for (const FNetViewer& CurViewer : ConnectionViewers)
+	{
+		FLastLocationGatherInfo* LastInfoForViewer = LastGatherLocations.FindByKey<UNetConnection*>(CurViewer.Connection);
+		if (LastInfoForViewer != nullptr)
+		{
+			OnUpdateViewerLocation(LastInfoForViewer, CurViewer, DestructionSettings);
+		}
+		else
+		{
+			// We need to add this viewer to the last gather locations
+			LastGatherLocations.Add(FLastLocationGatherInfo(CurViewer.Connection, CurViewer.ViewLocation));
+		}
+	}
+
+	// Clean up any dead entries in the last gather array
+	LastGatherLocations.RemoveAll([&](FLastLocationGatherInfo& CurGatherInfo) {
+		return CurGatherInfo.Connection == nullptr;
+	});
+}
+
+void UNetReplicationGraphConnection::OnUpdateViewerLocation(FLastLocationGatherInfo* LocationInfo, const FNetViewer& Viewer, const FReplicationGraphDestructionSettings& DestructionSettings )
+{
+	const bool bIgnoreDistanceCheck = DestructionSettings.OutOfRangeDistanceCheckThresholdSquared == 0.0f;
+
+	const float OutOfRangeDistanceSquared = FVector::DistSquared2D(Viewer.ViewLocation, LocationInfo->LastOutOfRangeLocationCheck);
+
+	// Test all accumulated out of range actors only once the viewer has gone far enough from the last check
+	if( bIgnoreDistanceCheck || (OutOfRangeDistanceSquared > DestructionSettings.OutOfRangeDistanceCheckThresholdSquared) )
+	{
+		CSV_SCOPED_TIMING_STAT_EXCLUSIVE(OnUpdateViewerLocation_TestDestroyedActors);
+
+		LocationInfo->LastOutOfRangeLocationCheck = Viewer.ViewLocation;
+
+		for (int32 Index=OutOfRangeDestroyedActors.Num()-1; Index >=0; --Index)
+		{
+			FCachedDestructInfo& CachedInfo = OutOfRangeDestroyedActors[Index];
+
+			const float ActorDistSquared = FVector::DistSquared2D(CachedInfo.CachedPosition, Viewer.ViewLocation);
+				
+			if (ActorDistSquared < DestructionSettings.DestructInfoMaxDistanceSquared)
+			{
+				// Swap the info into the Pending List to get it replicated
+				PendingDestructInfoList.Emplace(MoveTemp(CachedInfo));
+				OutOfRangeDestroyedActors.RemoveAtSwap(Index, 1, false);
+			}
+		}
+	}
+
+	LocationInfo->LastLocation = Viewer.ViewLocation;
+}
+
 // --------------------------------------------------------------------------------------------------------------------------------------------
 // --------------------------------------------------------------------------------------------------------------------------------------------
 // --------------------------------------------------------------------------------------------------------------------------------------------
diff --git a/Engine/Plugins/Runtime/ReplicationGraph/Source/Public/ReplicationGraph.h b/Engine/Plugins/Runtime/ReplicationGraph/Source/Public/ReplicationGraph.h
--- a/Engine/Plugins/Runtime/ReplicationGraph/Source/Public/ReplicationGraph.h
+++ b/Engine/Plugins/Runtime/ReplicationGraph/Source/Public/ReplicationGraph.h
@@ -63,6 +63,8 @@ UReplicationGraph::InitConnectionGraphNodes
 
 #include "ReplicationGraph.generated.h"
 
+struct FReplicationGraphDestructionSettings;
+
 #define DO_ENABLE_REPGRAPH_DEBUG_ACTOR !(UE_BUILD_SHIPPING || UE_BUILD_TEST)
 
 UCLASS(transient, config=Engine)
@@ -1020,8 +1022,8 @@ struct FLastLocationGatherInfo
 {
 	GENERATED_BODY()
 
-	FLastLocationGatherInfo() : Connection(nullptr), LastLocation(FVector(ForceInitToZero)) {}
-	FLastLocationGatherInfo(const UNetConnection* InConnection, FVector InLastLocation) : Connection(InConnection), LastLocation(InLastLocation) {}
+	FLastLocationGatherInfo() : Connection(nullptr), LastLocation(ForceInitToZero), LastOutOfRangeLocationCheck(ForceInitToZero) {}
+	FLastLocationGatherInfo(const UNetConnection* InConnection, FVector InLastLocation) : Connection(InConnection), LastLocation(InLastLocation), LastOutOfRangeLocationCheck(InLastLocation)  {}
 
 	UPROPERTY()
 	const UNetConnection* Connection;
@@ -1029,6 +1031,9 @@ struct FLastLocationGatherInfo
 	UPROPERTY()
 	FVector LastLocation;
 
+	UPROPERTY()
+	FVector LastOutOfRangeLocationCheck;
+
 	bool operator==(UNetConnection* Other) const
 	{
 		return Connection == Other;
@@ -1139,6 +1144,12 @@ private:
 	
 	int64 ReplicateDormantDestructionInfos();
 
+	/** Update the last location for viewers on a connection */
+	void UpdateGatherLocationsForConnection(const FNetViewerArray& ConnectionViewers, const FReplicationGraphDestructionSettings& DestructionSettings);
+
+    /** Update the location info of a specific viewer */
+	void OnUpdateViewerLocation(FLastLocationGatherInfo* LocationInfo, const FNetViewer& Viewer, const FReplicationGraphDestructionSettings& DestructionSettings);
+
 	UPROPERTY()
 	TArray<UReplicationGraphNode*> ConnectionGraphNodes;
 
@@ -1149,6 +1160,7 @@ private:
 	struct FCachedDestructInfo
 	{
 		FCachedDestructInfo(FActorDestructionInfo* InDestructInfo) : DestructionInfo(InDestructInfo), CachedPosition(InDestructInfo->DestroyedPosition) {}
+		bool operator==(const FCachedDestructInfo& rhs) const { return DestructionInfo == rhs.DestructionInfo; };
 		bool operator==(const FActorDestructionInfo* InDestructInfo) const { return InDestructInfo == DestructionInfo; };
 		
 		FActorDestructionInfo* DestructionInfo;
@@ -1164,8 +1176,17 @@ private:
 		}
 	};
 
+	/** 
+	* List of destroyed actors that were too far from the connection to be relevant.
+	* Is periodically evaluated when the viewer crosses a specific distance.
+	*/
+	TArray<FCachedDestructInfo> OutOfRangeDestroyedActors;
+
+	/** List of destroyed actors that need to be replicated */
 	TArray<FCachedDestructInfo> PendingDestructInfoList;
-	TSet<FActorDestructionInfo*> TrackedDestructionInfoPtrs; // Set used to guard against double adds into PendingDestructInfoList
+
+	/** Set used to guard against double adds into PendingDestructInfoList */
+	TSet<FActorDestructionInfo*> TrackedDestructionInfoPtrs;
 
 	struct FCachedDormantDestructInfo
 	{
@@ -1175,6 +1196,7 @@ private:
 		FString PathName;
 	};
 
+	/** List of dormant actors that should be removed from the client */
 	TArray<FCachedDormantDestructInfo> PendingDormantDestructList;
 };
 
@@ -1240,3 +1262,17 @@ public:
 
 	virtual UNetConnection* GetNetConnection() const override;
 };
+
+// --------------------------------------------------------------------------------------------------------------------------------------------
+// --------------------------------------------------------------------------------------------------------------------------------------------
+// --------------------------------------------------------------------------------------------------------------------------------------------
+struct FReplicationGraphDestructionSettings
+{
+	FReplicationGraphDestructionSettings(float InDestructInfoMaxDistanceSquared, float InOutOfRangeDistanceCheckThresholdSquared)
+		: DestructInfoMaxDistanceSquared(InDestructInfoMaxDistanceSquared)
+		, OutOfRangeDistanceCheckThresholdSquared(InOutOfRangeDistanceCheckThresholdSquared)
+	{ }
+
+	float DestructInfoMaxDistanceSquared;
+	float OutOfRangeDistanceCheckThresholdSquared;
+};
