Fix missing call to Teardown when removing ReplicationGraphConnection's
Improve cleanup of ReplicationGraphConnection in InitializeForWorld

diff --git a/Engine/Plugins/Runtime/ReplicationGraph/Source/Private/ReplicationGraph.cpp b/Engine/Plugins/Runtime/ReplicationGraph/Source/Private/ReplicationGraph.cpp
--- a/Engine/Plugins/Runtime/ReplicationGraph/Source/Private/ReplicationGraph.cpp
+++ b/Engine/Plugins/Runtime/ReplicationGraph/Source/Private/ReplicationGraph.cpp
@@ -446,6 +446,7 @@ void UReplicationGraph::RemoveClientConnection(UNetConnection* NetConnection)
 			if (ConnectionManager->NetConnection == NetConnection)
 			{
 				ensure(!bFound);
+				ConnectionManager->TearDown();
 				Connections.RemoveAtSwap(idx, 1, false);
 				bFound = true;
 			}
@@ -505,6 +506,11 @@ void UReplicationGraph::InitializeForWorld(UWorld* World)
 	{
 		Manager->NotifyResetAllNetworkActors();
 	}
+
+	for (UNetReplicationGraphConnection* RepGraphConnection : Connections)
+	{
+		RepGraphConnection->NotifyResetAllNetworkActors();
+	}
 	
 	if (World)
 	{
@@ -2470,6 +2476,16 @@ void UNetReplicationGraphConnection::NotifyResetDestructionInfo()
 	OutOfRangeDestroyedActors.Reset();
 }
 
+void UNetReplicationGraphConnection::NotifyResetAllNetworkActors()
+{
+	for (UReplicationGraphNode* Node : ConnectionGraphNodes)
+	{
+		Node->NotifyResetAllNetworkActors();
+	}
+
+	ActorInfoMap.ResetActorMap();
+}
+
 void UNetReplicationGraphConnection::GetClientVisibleLevelNames(TSet<FName>& OutLevelNames) const
 {
 	if (NetConnection == nullptr)
diff --git a/Engine/Plugins/Runtime/ReplicationGraph/Source/Public/ReplicationGraph.h b/Engine/Plugins/Runtime/ReplicationGraph/Source/Public/ReplicationGraph.h
--- a/Engine/Plugins/Runtime/ReplicationGraph/Source/Public/ReplicationGraph.h
+++ b/Engine/Plugins/Runtime/ReplicationGraph/Source/Public/ReplicationGraph.h
@@ -1091,8 +1091,6 @@ public:
 	/** Returns connection graph nodes. This is const so that you do not mutate the array itself. You should use AddConnectionGraphNode/RemoveConnectionGraphNode.  */
 	const TArray<UReplicationGraphNode*>& GetConnectionGraphNodes() const { return ConnectionGraphNodes; }
 
-	virtual void NotifyAddDormantDestructionInfo(AActor* Actor) override;
-
 	//~ Begin UObject Interface
 	virtual void Serialize(FArchive& Ar) override;
 	//~ End UObject Interface
@@ -1114,9 +1112,13 @@ public:
 
 	virtual void NotifyRemoveDestructionInfo(FActorDestructionInfo* DestructInfo) override;
 
+	virtual void NotifyAddDormantDestructionInfo(AActor* Actor) override;
+
 	virtual void NotifyResetDestructionInfo() override;
 	//~ End UReplicationConnectionDriver Interface
 
+	virtual void NotifyResetAllNetworkActors();
+
 	/** Generates a set of all the visible level names for this connection and its subconnections (if any) */
 	virtual void GetClientVisibleLevelNames(TSet<FName>& OutLevelNames) const;
 
diff --git a/Engine/Plugins/Runtime/ReplicationGraph/Source/Public/ReplicationGraphTypes.h b/Engine/Plugins/Runtime/ReplicationGraph/Source/Public/ReplicationGraphTypes.h
--- a/Engine/Plugins/Runtime/ReplicationGraph/Source/Public/ReplicationGraphTypes.h
+++ b/Engine/Plugins/Runtime/ReplicationGraph/Source/Public/ReplicationGraphTypes.h
@@ -955,6 +955,12 @@ struct FPerConnectionActorInfoMap
 		return ChannelMap.CreateIterator();
 	}
 
+	void ResetActorMap()
+	{
+		ActorMap.Reset();
+		ChannelMap.Reset();
+	}
+
 	int32 Num() const { return ActorMap.Num(); }
 
 	void CountBytes(FArchive& Ar) const
