RepGraph
* When modifying the actor-specific cull distance we also need to modify all existing connection-specific cull distance settings.
* Created SetAllCullDistanceSettingsForActor utility function to be used for such scenarios.
* Multicast cull distance check will now use the connection cull distance setting

diff --git a/Engine/Plugins/Runtime/ReplicationGraph/Source/Private/ReplicationGraph.cpp b/Engine/Plugins/Runtime/ReplicationGraph/Source/Private/ReplicationGraph.cpp
--- a/Engine/Plugins/Runtime/ReplicationGraph/Source/Private/ReplicationGraph.cpp
+++ b/Engine/Plugins/Runtime/ReplicationGraph/Source/Private/ReplicationGraph.cpp
@@ -1973,7 +1973,6 @@ bool UReplicationGraph::ProcessRemoteFunction(class AActor* Actor, UFunction* Fu
 
 		RepLayout->BuildSharedSerializationForRPC(Parameters);
 		FGlobalActorReplicationInfo& GlobalInfo = GlobalActorReplicationInfoMap.Get(Actor);
-		const float CullDistanceSquared = GlobalInfo.Settings.GetCullDistanceSquared();
 
 		bool ForceFlushNetDormancy = false;
 
@@ -2009,7 +2008,7 @@ bool UReplicationGraph::ProcessRemoteFunction(class AActor* Actor, UFunction* Fu
 				// if (Actor->NetDormancy > DORM_Awake)
 				{
 					bool ShouldOpenChannel = true;
-					if (CullDistanceSquared > 0.f)
+					if (ConnectionActorInfo.GetCullDistanceSquared() > 0.f)
 					{
 						ShouldOpenChannel = false;
 						if (ActorLocation.IsSet() == false)
@@ -2033,7 +2032,7 @@ bool UReplicationGraph::ProcessRemoteFunction(class AActor* Actor, UFunction* Fu
 						for (const FNetViewer& Viewer : ViewsToConsider)
 						{
 							const float DistSq = (ActorLocation.GetValue() - Viewer.ViewLocation).SizeSquared();
-							if (DistSq <= CullDistanceSquared)
+							if (DistSq <= ConnectionActorInfo.GetCullDistanceSquared())
 							{
 								ShouldOpenChannel = true;
 								break;
@@ -2181,6 +2180,20 @@ void UReplicationGraph::SetActorDiscoveryBudget(int32 ActorDiscoveryBudgetInKByt
 	UE_LOG(LogReplicationGraph, Display, TEXT("SetActorDiscoveryBudget set to %d kBps (%d bits per network tick)."), ActorDiscoveryBudgetInKBytesPerSec, ActorDiscoveryMaxBitsPerFrame);
 }
 
+void UReplicationGraph::SetAllCullDistanceSettingsForActor(const FActorRepListType& Actor, float CullDistanceSquared)
+{
+	FGlobalActorReplicationInfo& GlobalInfo = GlobalActorReplicationInfoMap.Get(Actor);
+	GlobalInfo.Settings.SetCullDistanceSquared(CullDistanceSquared);
+
+	for (UNetReplicationGraphConnection* RepGraphConnection : Connections)
+	{
+		if (FConnectionReplicationActorInfo* ConnectionActorInfo = RepGraphConnection->ActorInfoMap.Find(Actor))
+		{
+			ConnectionActorInfo->SetCullDistanceSquared(CullDistanceSquared);
+		}
+	}
+}
+
 void UReplicationGraph::NotifyConnectionSaturated(UNetReplicationGraphConnection& Connection)
 {
 	bWasConnectionSaturated = true;
diff --git a/Engine/Plugins/Runtime/ReplicationGraph/Source/Public/ReplicationGraph.h b/Engine/Plugins/Runtime/ReplicationGraph/Source/Public/ReplicationGraph.h
--- a/Engine/Plugins/Runtime/ReplicationGraph/Source/Public/ReplicationGraph.h
+++ b/Engine/Plugins/Runtime/ReplicationGraph/Source/Public/ReplicationGraph.h
@@ -844,6 +844,9 @@ public:
 
 	void SetActorDiscoveryBudget(int32 ActorDiscoveryBudgetInKBytesPerSec);
 
+	/** Sets the global and connection-specific cull distance setting of this actor */
+	void SetAllCullDistanceSettingsForActor(const FActorRepListType& Actor, float CullDistanceSquared);
+
 	// --------------------------------------------------------------
 
 	/** Creates a new node for the graph. This and UReplicationNode::CreateChildNode should be the only things that create the graph node UObjects */
