Fix ReplicationGraph ignoring traffic for the saturation calculations when ActorDiscoveryMaxBitsPerFrame is set to 0
This caused the ReplicationGraph to send more data than was budgeted by the MaxInternetClientRate

[RN] minor engine/network

diff --git a/Engine/Plugins/Runtime/ReplicationGraph/Source/Private/ReplicationGraph.cpp b/Engine/Plugins/Runtime/ReplicationGraph/Source/Private/ReplicationGraph.cpp
--- a/Engine/Plugins/Runtime/ReplicationGraph/Source/Private/ReplicationGraph.cpp
+++ b/Engine/Plugins/Runtime/ReplicationGraph/Source/Private/ReplicationGraph.cpp
@@ -899,7 +899,6 @@ int32 UReplicationGraph::ServerReplicateActors(float DeltaSeconds)
 
 		const FReplicationGraphDestructionSettings DestructionSettings(DestructInfoMaxDistanceSquared, CVar_RepGraph_OutOfRangeDistanceCheckRatio * DestructInfoMaxDistanceSquared);
 
-		FBitWriter& ConnectionSendBuffer = NetConnection->SendBuffer; // unused
 		ConnectionManager->QueuedBitsForActorDiscovery = 0;
 
 		// --------------------------------------------------------------------------------------------------------------
@@ -1826,7 +1825,7 @@ int64 UReplicationGraph::ReplicateSingleActor(AActor* Actor, FConnectionReplicat
 	}
 
 	// Optional budget for actor discovery traffic
-	if (!bIsActorDiscoveryBudgetFull)
+	if (bIsTrafficActorDiscovery && !bIsActorDiscoveryBudgetFull)
 	{
 		ConnectionManager.QueuedBitsForActorDiscovery += BitsWritten;
 
