PR #6130: Bug fix: UReplicationGraphNode_ActorList::NotifyRemoveNetworkActor (Contributed by [REDACTED_PERSON])

diff --git a/Engine/Plugins/Runtime/ReplicationGraph/Source/Private/ReplicationGraph.cpp b/Engine/Plugins/Runtime/ReplicationGraph/Source/Private/ReplicationGraph.cpp
--- a/Engine/Plugins/Runtime/ReplicationGraph/Source/Private/ReplicationGraph.cpp
+++ b/Engine/Plugins/Runtime/ReplicationGraph/Source/Private/ReplicationGraph.cpp
@@ -2739,14 +2739,9 @@ bool UReplicationGraphNode_ActorList::NotifyRemoveNetworkActor(const FNewReplica
 
 	if (ActorInfo.StreamingLevelName == NAME_None)
 	{
-		if (!ReplicationActorList.Remove(ActorInfo.Actor) && bWarnIfNotFound)
-		{
-			UE_LOG(LogReplicationGraph, Warning, TEXT("Attempted to remove %s from list %s but it was not found. (StreamingLevelName == NAME_None)"), *GetActorRepListTypeDebugString(ActorInfo.Actor), *GetFullName());
-		}
-		else
-		{
-			bRemovedSomething = true;
-		}
+		bRemovedSomething = ReplicationActorList.Remove(ActorInfo.Actor);
+
+		UE_CLOG(!bRemovedSomething && bWarnIfNotFound, LogReplicationGraph, Warning, TEXT("Attempted to remove %s from list %s but it was not found. (StreamingLevelName == NAME_None)"), *GetActorRepListTypeDebugString(ActorInfo.Actor), *GetFullName());
 
 		if (CVar_RepGraph_Verify)
 		{
