AddNetworkActor now checks the driver ShouldReplicateActor the way the network object list does

diff --git a/Engine/Plugins/Runtime/ReplicationGraph/Source/Private/ReplicationGraph.cpp b/Engine/Plugins/Runtime/ReplicationGraph/Source/Private/ReplicationGraph.cpp
--- a/Engine/Plugins/Runtime/ReplicationGraph/Source/Private/ReplicationGraph.cpp
+++ b/Engine/Plugins/Runtime/ReplicationGraph/Source/Private/ReplicationGraph.cpp
@@ -535,6 +535,11 @@ void UReplicationGraph::AddNetworkActor(AActor* Actor)
 		return;
 	}
 
+	if (NetDriver && !NetDriver->ShouldReplicateActor(Actor))
+	{
+		return;
+	}
+
 	bool bWasAlreadyThere = false;
 	ActiveNetworkActors.Add(Actor, &bWasAlreadyThere);
 	if (bWasAlreadyThere)
