Automated wrapped object pointer upgrade for Engine + ShooterGame + ShooterGame referenced plugins

diff --git a/Engine/Plugins/FX/Niagara/Source/Niagara/Public/NiagaraMeshRendererProperties.h b/Engine/Plugins/FX/Niagara/Source/Niagara/Public/NiagaraMeshRendererProperties.h
--- a/Engine/Plugins/FX/Niagara/Source/Niagara/Public/NiagaraMeshRendererProperties.h
+++ b/Engine/Plugins/FX/Niagara/Source/Niagara/Public/NiagaraMeshRendererProperties.h
@@ -61,7 +61,7 @@ public:
 
 	/** Use this UMaterialInterface if set to a valid value. This will be subordinate to UserParamBinding if it is set to a valid user variable.*/
 	UPROPERTY(EditAnywhere, Category = "Mesh Rendering")
-	UMaterialInterface* ExplicitMat;
+	TObjectPtr<UMaterialInterface> ExplicitMat;
 
 	/** Use the UMaterialInterface bound to this user variable if it is set to a valid value. If this is bound to a valid value and ExplicitMat is also set, UserParamBinding wins.*/
 	UPROPERTY(EditAnywhere, Category = "Mesh Rendering")
@@ -109,7 +109,7 @@ struct NIAGARA_API FNiagaraMeshRendererMeshProperties
 
 	/** The mesh to use when rendering this slot */
 	UPROPERTY(EditAnywhere, Category = "Mesh")
-	UStaticMesh* Mesh;
+	TObjectPtr<UStaticMesh> Mesh;
 
 	/** Scale of the mesh */
 	UPROPERTY(EditAnywhere, Category = "Mesh")
@@ -312,7 +312,7 @@ public:
 	 * will not be found or used.
 	 */
 	UPROPERTY(EditAnywhere, Category = "Mesh Flipbook", meta = (EditCondition = "bEnableMeshFlipbook"))
-	UStaticMesh* FirstFlipbookFrame;
+	TObjectPtr<UStaticMesh> FirstFlipbookFrame;
 
 	/**
 	 * Provides the format of the suffix of the names of the static meshes when searching for flipbook frames. "{frame_number}" is used to mark
@@ -353,7 +353,7 @@ private:
 
 	// These properties are deprecated and moved to FNiagaraMeshRendererMeshProperties
 	UPROPERTY()
-	UStaticMesh* ParticleMesh_DEPRECATED;
+	TObjectPtr<UStaticMesh> ParticleMesh_DEPRECATED;
 
 	UPROPERTY()
 	FVector PivotOffset_DEPRECATED;
diff --git a/Engine/Plugins/FX/Niagara/Source/Niagara/Public/NiagaraParameterStore.h b/Engine/Plugins/FX/Niagara/Source/Niagara/Public/NiagaraParameterStore.h
--- a/Engine/Plugins/FX/Niagara/Source/Niagara/Public/NiagaraParameterStore.h
+++ b/Engine/Plugins/FX/Niagara/Source/Niagara/Public/NiagaraParameterStore.h
@@ -162,7 +162,7 @@ struct NIAGARA_API FNiagaraParameterStore
 private:
 	/** Owner of this store. Used to provide an outer to data interfaces in this store. */
 	UPROPERTY(Transient)
-	UObject* Owner;
+	TObjectPtr<UObject> Owner;
 	
 #if WITH_EDITORONLY_DATA
 	/** Map from parameter defs to their offset in the data table or the data interface. TODO: Separate out into a layout and instance class to reduce duplicated data for this?  */
@@ -181,11 +181,11 @@ private:
 	
 	/** Data interfaces for this script. Possibly overridden with externally owned interfaces. Also indexed by ParameterOffsets. */
 	UPROPERTY()
-	TArray<UNiagaraDataInterface*> DataInterfaces;
+	TArray<TObjectPtr<UNiagaraDataInterface>> DataInterfaces;
 
 	/** UObjects referenced by this store. Also indexed by ParameterOffsets.*/
 	UPROPERTY()
-	TArray<UObject*> UObjects;
+	TArray<TObjectPtr<UObject>> UObjects;
 
 	/** Bindings between this parameter store and others we push data into when we tick. */
 	typedef TPair<FNiagaraParameterStore*, FNiagaraParameterStoreBinding> BindingPair;
diff --git a/Engine/Plugins/FX/Niagara/Source/Niagara/Public/NiagaraPreviewGrid.h b/Engine/Plugins/FX/Niagara/Source/Niagara/Public/NiagaraPreviewGrid.h
--- a/Engine/Plugins/FX/Niagara/Source/Niagara/Public/NiagaraPreviewGrid.h
+++ b/Engine/Plugins/FX/Niagara/Source/Niagara/Public/NiagaraPreviewGrid.h
@@ -221,18 +221,18 @@ private:
 public:
 
 	UPROPERTY(EditAnywhere, Category = Preview)
-	UNiagaraSystem* System;
+	TObjectPtr<UNiagaraSystem> System;
 	
 	UPROPERTY(EditAnywhere, Category = Preview)
 	ENiagaraPreviewGridResetMode ResetMode;
 
 	/** Object controlling behavior varying on the X axis. */
 	UPROPERTY(EditAnywhere, Category = Preview, Instanced)
-	UNiagaraPreviewAxis* PreviewAxisX;
+	TObjectPtr<UNiagaraPreviewAxis> PreviewAxisX;
 
 	/** Object controlling behavior varying on the Y axis. */
 	UPROPERTY(EditAnywhere, Category = Preview, Instanced)
-	UNiagaraPreviewAxis* PreviewAxisY;
+	TObjectPtr<UNiagaraPreviewAxis> PreviewAxisY;
 
 	/** Class used to for previews in this grid. */
 	UPROPERTY(EditAnywhere, Category = Preview)
@@ -256,16 +256,16 @@ private:
 	int32 NumY;
 
 	UPROPERTY(transient)
-	TArray<UChildActorComponent*> PreviewComponents;
+	TArray<TObjectPtr<UChildActorComponent>> PreviewComponents;
 
 #if WITH_EDITORONLY_DATA
 	// Reference to sprite visualization component
 	UPROPERTY()
-	class UBillboardComponent* SpriteComponent;
+	TObjectPtr<class UBillboardComponent> SpriteComponent;
 
 	// Reference to arrow visualization component
 	UPROPERTY()
-	class UArrowComponent* ArrowComponent;
+	TObjectPtr<class UArrowComponent> ArrowComponent;
 #endif
 
 
diff --git a/Engine/Plugins/FX/Niagara/Source/Niagara/Public/NiagaraRibbonRendererProperties.h b/Engine/Plugins/FX/Niagara/Source/Niagara/Public/NiagaraRibbonRendererProperties.h
--- a/Engine/Plugins/FX/Niagara/Source/Niagara/Public/NiagaraRibbonRendererProperties.h
+++ b/Engine/Plugins/FX/Niagara/Source/Niagara/Public/NiagaraRibbonRendererProperties.h
@@ -183,7 +183,7 @@ public:
 	//UNiagaraRendererProperties Interface END
 
 	UPROPERTY(EditAnywhere, Category = "Ribbon Rendering")
-	UMaterialInterface* Material;
+	TObjectPtr<UMaterialInterface> Material;
 
 	/** Use the UMaterialInterface bound to this user variable if it is set to a valid value. If this is bound to a valid value and Material is also set, UserParamBinding wins.*/
 	UPROPERTY(EditAnywhere, Category = "Ribbon Rendering")
diff --git a/Engine/Plugins/FX/Niagara/Source/Niagara/Public/NiagaraScalabilityManager.h b/Engine/Plugins/FX/Niagara/Source/Niagara/Public/NiagaraScalabilityManager.h
--- a/Engine/Plugins/FX/Niagara/Source/Niagara/Public/NiagaraScalabilityManager.h
+++ b/Engine/Plugins/FX/Niagara/Source/Niagara/Public/NiagaraScalabilityManager.h
@@ -138,10 +138,10 @@ struct FNiagaraScalabilityManager
 	GENERATED_USTRUCT_BODY()
 
 	UPROPERTY(transient)
-	UNiagaraEffectType* EffectType;
+	TObjectPtr<UNiagaraEffectType> EffectType;
 
 	UPROPERTY(transient)
-	TArray<UNiagaraComponent*>  ManagedComponents;
+	TArray<TObjectPtr<UNiagaraComponent>>  ManagedComponents;
 
 	TArray<FNiagaraScalabilityState> State;
 
diff --git a/Engine/Plugins/FX/Niagara/Source/Niagara/Public/NiagaraSettings.h b/Engine/Plugins/FX/Niagara/Source/Niagara/Public/NiagaraSettings.h
--- a/Engine/Plugins/FX/Niagara/Source/Niagara/Public/NiagaraSettings.h
+++ b/Engine/Plugins/FX/Niagara/Source/Niagara/Public/NiagaraSettings.h
@@ -107,5 +107,5 @@ protected:
 
 private:
 	UPROPERTY(transient)
-	mutable UNiagaraEffectType* DefaultEffectTypePtr;
+	mutable TObjectPtr<UNiagaraEffectType> DefaultEffectTypePtr;
 };
\ No newline at end of file
diff --git a/Engine/Plugins/FX/Niagara/Source/Niagara/Public/NiagaraSimulationStageBase.h b/Engine/Plugins/FX/Niagara/Source/Niagara/Public/NiagaraSimulationStageBase.h
--- a/Engine/Plugins/FX/Niagara/Source/Niagara/Public/NiagaraSimulationStageBase.h
+++ b/Engine/Plugins/FX/Niagara/Source/Niagara/Public/NiagaraSimulationStageBase.h
@@ -25,7 +25,7 @@ public:
 	}
 
 	UPROPERTY()
-	UNiagaraScript* Script;
+	TObjectPtr<UNiagaraScript> Script;
 
 	UPROPERTY(EditAnywhere, Category = "Simulation Stage")
 	FName SimulationStageName;
diff --git a/Engine/Plugins/FX/Niagara/Source/Niagara/Public/NiagaraSpriteRendererProperties.h b/Engine/Plugins/FX/Niagara/Source/Niagara/Public/NiagaraSpriteRendererProperties.h
--- a/Engine/Plugins/FX/Niagara/Source/Niagara/Public/NiagaraSpriteRendererProperties.h
+++ b/Engine/Plugins/FX/Niagara/Source/Niagara/Public/NiagaraSpriteRendererProperties.h
@@ -113,7 +113,7 @@ public:
 
 	/** The material used to render the particle. Note that it must have the Use with Niagara Sprites flag checked.*/
 	UPROPERTY(EditAnywhere, Category = "Sprite Rendering")
-	UMaterialInterface* Material;
+	TObjectPtr<UMaterialInterface> Material;
 
 	/** Whether or not to draw a single element for the Emitter or to draw the particles.*/
 	UPROPERTY(EditAnywhere, Category = "Sprite Rendering")
@@ -275,7 +275,7 @@ public:
 
 	/** Texture to generate bounding geometry from.	*/
 	UPROPERTY(EditAnywhere, Category="Cutout", meta = (EditCondition = "!bUseMaterialCutoutTexture"))
-	UTexture2D* CutoutTexture;
+	TObjectPtr<UTexture2D> CutoutTexture;
 	
 	/**
 	* More bounding vertices results in reduced overdraw, but adds more triangle overhead.
diff --git a/Engine/Plugins/FX/Niagara/Source/Niagara/Public/NiagaraTypes.h b/Engine/Plugins/FX/Niagara/Source/Niagara/Public/NiagaraTypes.h
--- a/Engine/Plugins/FX/Niagara/Source/Niagara/Public/NiagaraTypes.h
+++ b/Engine/Plugins/FX/Niagara/Source/Niagara/Public/NiagaraTypes.h
@@ -957,7 +957,7 @@ public:
 	In occasional situations this may be a UClass when we're dealing with DataInterface etc.
 	*/
 	UPROPERTY(EditAnywhere, Category=Type)
-	UObject* ClassStructOrEnum;
+	TObjectPtr<UObject> ClassStructOrEnum;
 
 	// See enumeration FUnderlyingType for possible values
 	UPROPERTY(EditAnywhere, Category=Type)
@@ -972,10 +972,10 @@ private:
 
 #if WITH_EDITORONLY_DATA
 	UPROPERTY()
-	UStruct* Struct_DEPRECATED;
+	TObjectPtr<UStruct> Struct_DEPRECATED;
 
 	UPROPERTY()
-	UEnum* Enum_DEPRECATED;
+	TObjectPtr<UEnum> Enum_DEPRECATED;
 #endif
 
 public:
diff --git a/Engine/Plugins/FX/Niagara/Source/Niagara/Public/NiagaraVariant.h b/Engine/Plugins/FX/Niagara/Source/Niagara/Public/NiagaraVariant.h
--- a/Engine/Plugins/FX/Niagara/Source/Niagara/Public/NiagaraVariant.h
+++ b/Engine/Plugins/FX/Niagara/Source/Niagara/Public/NiagaraVariant.h
@@ -46,10 +46,10 @@ struct NIAGARA_API FNiagaraVariant
 
 private:
 	UPROPERTY(EditAnywhere, Category=Variant, Instanced)
-	UObject* Object;
+	TObjectPtr<UObject> Object;
 
 	UPROPERTY(EditAnywhere, Category=Variant, Instanced)
-	UNiagaraDataInterface* DataInterface;
+	TObjectPtr<UNiagaraDataInterface> DataInterface;
 
 	UPROPERTY(EditAnywhere, Category=Variant)
 	TArray<uint8> Bytes;
diff --git a/Engine/Plugins/FX/Niagara/Source/NiagaraAnimNotifies/Public/AnimNotifyState_TimedNiagaraEffect.h b/Engine/Plugins/FX/Niagara/Source/NiagaraAnimNotifies/Public/AnimNotifyState_TimedNiagaraEffect.h
--- a/Engine/Plugins/FX/Niagara/Source/NiagaraAnimNotifies/Public/AnimNotifyState_TimedNiagaraEffect.h
+++ b/Engine/Plugins/FX/Niagara/Source/NiagaraAnimNotifies/Public/AnimNotifyState_TimedNiagaraEffect.h
@@ -22,7 +22,7 @@ class NIAGARAANIMNOTIFIES_API UAnimNotifyState_TimedNiagaraEffect : public UAnim
 
 	// The niagara system template to use when spawning the niagara component
 	UPROPERTY(EditAnywhere, Category = NiagaraSystem, meta = (ToolTip = "The niagara system to spawn for the notify state"))
-	UNiagaraSystem* Template;
+	TObjectPtr<UNiagaraSystem> Template;
 
 	// The socket within our mesh component to attach to when we spawn the Niagara component
 	UPROPERTY(EditAnywhere, Category = NiagaraSystem, meta = (ToolTip = "The socket or bone to attach the system to", AnimNotifyBoneName = "true"))
diff --git a/Engine/Plugins/FX/Niagara/Source/NiagaraAnimNotifies/Public/AnimNotify_PlayNiagaraEffect.h b/Engine/Plugins/FX/Niagara/Source/NiagaraAnimNotifies/Public/AnimNotify_PlayNiagaraEffect.h
--- a/Engine/Plugins/FX/Niagara/Source/NiagaraAnimNotifies/Public/AnimNotify_PlayNiagaraEffect.h
+++ b/Engine/Plugins/FX/Niagara/Source/NiagaraAnimNotifies/Public/AnimNotify_PlayNiagaraEffect.h
@@ -43,7 +43,7 @@ public:
 
 	// Niagara System to Spawn
 	UPROPERTY(EditAnywhere, BlueprintReadWrite, Category = "AnimNotify", meta = (DisplayName = "Niagara System"))
-	UNiagaraSystem* Template;
+	TObjectPtr<UNiagaraSystem> Template;
 
 	// Location offset from the socket
 	UPROPERTY(EditAnywhere, BlueprintReadWrite, Category = "AnimNotify")
diff --git a/Engine/Plugins/FX/Niagara/Source/NiagaraEditor/Private/NiagaraEmitterEditorData.h b/Engine/Plugins/FX/Niagara/Source/NiagaraEditor/Private/NiagaraEmitterEditorData.h
--- a/Engine/Plugins/FX/Niagara/Source/NiagaraEditor/Private/NiagaraEmitterEditorData.h
+++ b/Engine/Plugins/FX/Niagara/Source/NiagaraEditor/Private/NiagaraEmitterEditorData.h
@@ -27,7 +27,7 @@ public:
 
 private:
 	UPROPERTY(Instanced)
-	UNiagaraStackEditorData* StackEditorData;
+	TObjectPtr<UNiagaraStackEditorData> StackEditorData;
 
 	UPROPERTY()
 	float PlaybackRangeMin;
diff --git a/Engine/Plugins/FX/Niagara/Source/NiagaraEditor/Private/NiagaraNodeEmitter.h b/Engine/Plugins/FX/Niagara/Source/NiagaraEditor/Private/NiagaraNodeEmitter.h
--- a/Engine/Plugins/FX/Niagara/Source/NiagaraEditor/Private/NiagaraNodeEmitter.h
+++ b/Engine/Plugins/FX/Niagara/Source/NiagaraEditor/Private/NiagaraNodeEmitter.h
@@ -78,7 +78,7 @@ private:
 
 	/** The System that owns the emitter which is represented by this node. */
 	UPROPERTY(meta = (SkipForCompileHash = "true"))
-	UNiagaraSystem* OwnerSystem;
+	TObjectPtr<UNiagaraSystem> OwnerSystem;
 
 	/** The id of the emitter handle which points to the emitter represented by this node. */
 	UPROPERTY()
diff --git a/Engine/Plugins/FX/Niagara/Source/NiagaraEditor/Private/NiagaraNodeStaticSwitch.h b/Engine/Plugins/FX/Niagara/Source/NiagaraEditor/Private/NiagaraNodeStaticSwitch.h
--- a/Engine/Plugins/FX/Niagara/Source/NiagaraEditor/Private/NiagaraNodeStaticSwitch.h
+++ b/Engine/Plugins/FX/Niagara/Source/NiagaraEditor/Private/NiagaraNodeStaticSwitch.h
@@ -28,7 +28,7 @@ struct FStaticSwitchTypeData
 
 	/** If the type is enum, this is the enum being switched on, otherwise it holds no sensible value */
 	UPROPERTY()
-	UEnum* Enum;
+	TObjectPtr<UEnum> Enum;
 
 	/** If set, then this switch is not exposed but will rather be evaluated by the given compile-time constant */
 	UPROPERTY()
diff --git a/Engine/Plugins/FX/Niagara/Source/NiagaraEditor/Private/Sequencer/NiagaraSequence/MovieSceneNiagaraEmitterTrack.h b/Engine/Plugins/FX/Niagara/Source/NiagaraEditor/Private/Sequencer/NiagaraSequence/MovieSceneNiagaraEmitterTrack.h
--- a/Engine/Plugins/FX/Niagara/Source/NiagaraEditor/Private/Sequencer/NiagaraSequence/MovieSceneNiagaraEmitterTrack.h
+++ b/Engine/Plugins/FX/Niagara/Source/NiagaraEditor/Private/Sequencer/NiagaraSequence/MovieSceneNiagaraEmitterTrack.h
@@ -101,7 +101,7 @@ private:
 	TWeakPtr<FNiagaraEmitterHandleViewModel> EmitterHandleViewModel;
 
 	UPROPERTY()
-	TArray<UMovieSceneSection*> Sections;
+	TArray<TObjectPtr<UMovieSceneSection>> Sections;
 
 	UPROPERTY()
 	bool bSectionsWereModified;
diff --git a/Engine/Plugins/FX/Niagara/Source/NiagaraEditor/Private/Sequencer/NiagaraSequence/NiagaraSequence.h b/Engine/Plugins/FX/Niagara/Source/NiagaraEditor/Private/Sequencer/NiagaraSequence/NiagaraSequence.h
--- a/Engine/Plugins/FX/Niagara/Source/NiagaraEditor/Private/Sequencer/NiagaraSequence/NiagaraSequence.h
+++ b/Engine/Plugins/FX/Niagara/Source/NiagaraEditor/Private/Sequencer/NiagaraSequence/NiagaraSequence.h
@@ -34,7 +34,7 @@ public:
 private:
 	/** Pointer to the movie scene that controls this sequence. */
 	UPROPERTY()
-	UMovieScene* MovieScene;
+	TObjectPtr<UMovieScene> MovieScene;
 
 	/** The system view model which owns this niagara sequence. */
 	FNiagaraSystemViewModel* SystemViewModel;
diff --git a/Engine/Plugins/FX/Niagara/Source/NiagaraEditor/Public/EdGraphSchema_Niagara.h b/Engine/Plugins/FX/Niagara/Source/NiagaraEditor/Public/EdGraphSchema_Niagara.h
--- a/Engine/Plugins/FX/Niagara/Source/NiagaraEditor/Public/EdGraphSchema_Niagara.h
+++ b/Engine/Plugins/FX/Niagara/Source/NiagaraEditor/Public/EdGraphSchema_Niagara.h
@@ -22,7 +22,7 @@ struct NIAGARAEDITOR_API FNiagaraSchemaAction_NewNode : public FEdGraphSchemaAct
 
 	/** Template of node we want to create */
 	UPROPERTY()
-	class UEdGraphNode* NodeTemplate;
+	TObjectPtr<class UEdGraphNode> NodeTemplate;
 
 	UPROPERTY()
 	FName InternalName;
diff --git a/Engine/Plugins/FX/Niagara/Source/NiagaraEditor/Public/NiagaraClipboard.h b/Engine/Plugins/FX/Niagara/Source/NiagaraEditor/Public/NiagaraClipboard.h
--- a/Engine/Plugins/FX/Niagara/Source/NiagaraEditor/Public/NiagaraClipboard.h
+++ b/Engine/Plugins/FX/Niagara/Source/NiagaraEditor/Public/NiagaraClipboard.h
@@ -63,13 +63,13 @@ public:
 	FName Linked;
 
 	UPROPERTY()
-	UNiagaraDataInterface* Data;
+	TObjectPtr<UNiagaraDataInterface> Data;
 
 	UPROPERTY()
 	FString Expression;
 
 	UPROPERTY()
-	UNiagaraClipboardFunction* Dynamic;
+	TObjectPtr<UNiagaraClipboardFunction> Dynamic;
 
 	bool CopyValuesFrom(const UNiagaraClipboardFunctionInput* InOther);
 
@@ -115,7 +115,7 @@ public:
 	TArray<FString> AssignmentDefaults;
 
 	UPROPERTY()
-	TArray<const UNiagaraClipboardFunctionInput*> Inputs;
+	TArray<TObjectPtr<const UNiagaraClipboardFunctionInput>> Inputs;
 
 	UPROPERTY()
 	FOnPastedFunctionCallNode OnPastedFunctionCallNodeDelegate;
@@ -130,19 +130,19 @@ public:
 	static UNiagaraClipboardContent* Create();
 
 	UPROPERTY()
-	TArray<const UNiagaraClipboardFunction*> Functions;
+	TArray<TObjectPtr<const UNiagaraClipboardFunction>> Functions;
 
 	UPROPERTY()
-	TArray<const UNiagaraClipboardFunctionInput*> FunctionInputs;
+	TArray<TObjectPtr<const UNiagaraClipboardFunctionInput>> FunctionInputs;
 
 	UPROPERTY()
-	TArray<const UNiagaraRendererProperties*> Renderers;
+	TArray<TObjectPtr<const UNiagaraRendererProperties>> Renderers;
 
 	UPROPERTY()
-	TArray<const UNiagaraScript*> Scripts;
+	TArray<TObjectPtr<const UNiagaraScript>> Scripts;
 
 	UPROPERTY()
-	TArray<const UNiagaraScriptVariable*> ScriptVariables;
+	TArray<TObjectPtr<const UNiagaraScriptVariable>> ScriptVariables;
 };
 
 class NIAGARAEDITOR_API FNiagaraClipboard
diff --git a/Engine/Plugins/FX/Niagara/Source/NiagaraEditor/Public/NiagaraGraph.h b/Engine/Plugins/FX/Niagara/Source/NiagaraEditor/Public/NiagaraGraph.h
--- a/Engine/Plugins/FX/Niagara/Source/NiagaraEditor/Public/NiagaraGraph.h
+++ b/Engine/Plugins/FX/Niagara/Source/NiagaraEditor/Public/NiagaraGraph.h
@@ -62,7 +62,7 @@ public:
 	TArray<FNiagaraGraphParameterReference> ParameterReferences;
 
 	UPROPERTY()
-	const UNiagaraGraph* Graph;
+	TObjectPtr<const UNiagaraGraph> Graph;
 
 	/** Returns true if this parameter was initially created by the user. */
 	bool WasCreated() const;
@@ -109,7 +109,7 @@ public:
 
 	/** The traversal of output to input nodes for this graph. This is not a recursive traversal, it just includes nodes from this graph.*/
 	UPROPERTY()
-	TArray<UNiagaraNode*> Traversal;
+	TArray<TObjectPtr<UNiagaraNode>> Traversal;
 
 	void PostLoad(UObject* Owner);
 
@@ -451,7 +451,7 @@ private:
 
 	/** Storage of variables defined for use with this graph.*/
 	UPROPERTY()
-	mutable TMap<FNiagaraVariable, UNiagaraScriptVariable*> VariableToScriptVariable;
+	mutable TMap<FNiagaraVariable, TObjectPtr<UNiagaraScriptVariable>> VariableToScriptVariable;
 	
 	/** A map of parameters in the graph to their referencers. */
 	UPROPERTY(Transient)
diff --git a/Engine/Plugins/FX/Niagara/Source/NiagaraEditor/Public/NiagaraNodeDataSetBase.h b/Engine/Plugins/FX/Niagara/Source/NiagaraEditor/Public/NiagaraNodeDataSetBase.h
--- a/Engine/Plugins/FX/Niagara/Source/NiagaraEditor/Public/NiagaraNodeDataSetBase.h
+++ b/Engine/Plugins/FX/Niagara/Source/NiagaraEditor/Public/NiagaraNodeDataSetBase.h
@@ -26,7 +26,7 @@ public:
 	TArray<FString> VariableFriendlyNames;
 
 	UPROPERTY()
-	const UStruct* ExternalStructAsset;
+	TObjectPtr<const UStruct> ExternalStructAsset;
 
 	//~ Begin UObject Interface
 	virtual void PostEditChangeProperty(struct FPropertyChangedEvent& PropertyChangedEvent) override;
diff --git a/Engine/Plugins/FX/Niagara/Source/NiagaraEditor/Public/NiagaraNodeFunctionCall.h b/Engine/Plugins/FX/Niagara/Source/NiagaraEditor/Public/NiagaraNodeFunctionCall.h
--- a/Engine/Plugins/FX/Niagara/Source/NiagaraEditor/Public/NiagaraNodeFunctionCall.h
+++ b/Engine/Plugins/FX/Niagara/Source/NiagaraEditor/Public/NiagaraNodeFunctionCall.h
@@ -57,7 +57,7 @@ public:
 public:
 
 	UPROPERTY(EditAnywhere, Category = "Function")
-	UNiagaraScript* FunctionScript;
+	TObjectPtr<UNiagaraScript> FunctionScript;
 
 	/** 
 	 * A path to a script asset which can be used to assign the function script the first time that
@@ -166,7 +166,7 @@ protected:
 	FString FunctionDisplayName;
 
 	UPROPERTY(meta = (SkipForCompileHash="true"))
-	TMap<FGuid, UNiagaraMessageData*> MessageKeyToMessageMap;
+	TMap<FGuid, TObjectPtr<UNiagaraMessageData>> MessageKeyToMessageMap;
 
 	FOnInputsChanged OnInputsChangedDelegate;
 };
diff --git a/Engine/Plugins/FX/Niagara/Source/NiagaraEditor/Public/NiagaraNodeInput.h b/Engine/Plugins/FX/Niagara/Source/NiagaraEditor/Public/NiagaraNodeInput.h
--- a/Engine/Plugins/FX/Niagara/Source/NiagaraEditor/Public/NiagaraNodeInput.h
+++ b/Engine/Plugins/FX/Niagara/Source/NiagaraEditor/Public/NiagaraNodeInput.h
@@ -117,6 +117,6 @@ private:
 
 private:
 	UPROPERTY(meta = (SkipForCompileHash = "true"))
-	class UNiagaraDataInterface* DataInterface;
+	TObjectPtr<class UNiagaraDataInterface> DataInterface;
 };
 
diff --git a/Engine/Plugins/FX/Niagara/Source/NiagaraEditor/Public/NiagaraOverviewNode.h b/Engine/Plugins/FX/Niagara/Source/NiagaraEditor/Public/NiagaraOverviewNode.h
--- a/Engine/Plugins/FX/Niagara/Source/NiagaraEditor/Public/NiagaraOverviewNode.h
+++ b/Engine/Plugins/FX/Niagara/Source/NiagaraEditor/Public/NiagaraOverviewNode.h
@@ -46,7 +46,7 @@ public:
 
 private:
 	UPROPERTY()
-	UNiagaraSystem* OwningSystem;
+	TObjectPtr<UNiagaraSystem> OwningSystem;
 
 	UPROPERTY()
 	FGuid EmitterHandleGuid;
diff --git a/Engine/Plugins/FX/Niagara/Source/NiagaraEditor/Public/NiagaraScriptSource.h b/Engine/Plugins/FX/Niagara/Source/NiagaraEditor/Public/NiagaraScriptSource.h
--- a/Engine/Plugins/FX/Niagara/Source/NiagaraEditor/Public/NiagaraScriptSource.h
+++ b/Engine/Plugins/FX/Niagara/Source/NiagaraEditor/Public/NiagaraScriptSource.h
@@ -18,7 +18,7 @@ class UNiagaraScriptSource : public UNiagaraScriptSourceBase
 
 	/** Graph for particle update expression */
 	UPROPERTY()
-	class UNiagaraGraph*	NodeGraph;
+	TObjectPtr<class UNiagaraGraph>	NodeGraph;
 	
 	// UObject interface
 	virtual void PostLoad() override;
diff --git a/Engine/Plugins/FX/Niagara/Source/NiagaraEditor/Public/NiagaraSystemEditorData.h b/Engine/Plugins/FX/Niagara/Source/NiagaraEditor/Public/NiagaraSystemEditorData.h
--- a/Engine/Plugins/FX/Niagara/Source/NiagaraEditor/Public/NiagaraSystemEditorData.h
+++ b/Engine/Plugins/FX/Niagara/Source/NiagaraEditor/Public/NiagaraSystemEditorData.h
@@ -37,7 +37,7 @@ private:
 	FName FolderName;
 
 	UPROPERTY()
-	TArray<UNiagaraSystemEditorFolder*> ChildFolders;
+	TArray<TObjectPtr<UNiagaraSystemEditorFolder>> ChildFolders;
 
 	UPROPERTY()
 	TArray<FGuid> ChildEmitterHandleIds;
@@ -93,10 +93,10 @@ private:
 
 private:
 	UPROPERTY(Instanced)
-	UNiagaraSystemEditorFolder* RootFolder;
+	TObjectPtr<UNiagaraSystemEditorFolder> RootFolder;
 
 	UPROPERTY(Instanced)
-	UNiagaraStackEditorData* StackEditorData;
+	TObjectPtr<UNiagaraStackEditorData> StackEditorData;
 
 	UPROPERTY()
 	FTransform OwnerTransform;
@@ -109,7 +109,7 @@ private:
 
 	/** Graph presenting overview of the current system and its emitters. */
 	UPROPERTY()
-	UEdGraph* SystemOverviewGraph;
+	TObjectPtr<UEdGraph> SystemOverviewGraph;
 
 	UPROPERTY()
 	FNiagaraGraphViewSettings OverviewGraphViewSettings;
diff --git a/Engine/Plugins/FX/Niagara/Source/NiagaraEditor/Public/ViewModels/NiagaraSystemSelectionViewModel.h b/Engine/Plugins/FX/Niagara/Source/NiagaraEditor/Public/ViewModels/NiagaraSystemSelectionViewModel.h
--- a/Engine/Plugins/FX/Niagara/Source/NiagaraEditor/Public/ViewModels/NiagaraSystemSelectionViewModel.h
+++ b/Engine/Plugins/FX/Niagara/Source/NiagaraEditor/Public/ViewModels/NiagaraSystemSelectionViewModel.h
@@ -85,10 +85,10 @@ private:
 	TArray<FGuid> SelectedEmitterHandleIds;
 
 	UPROPERTY()
-	UNiagaraStackSelection* StackSelection;
+	TObjectPtr<UNiagaraStackSelection> StackSelection;
 
 	UPROPERTY()
-	UNiagaraStackViewModel* SelectionStackViewModel;
+	TObjectPtr<UNiagaraStackViewModel> SelectionStackViewModel;
 
 	FOnSelectionChanged OnEntrySelectionChangedDelegate;
 
diff --git a/Engine/Plugins/FX/Niagara/Source/NiagaraEditor/Public/ViewModels/Stack/NiagaraStackEmitterSettingsGroup.h b/Engine/Plugins/FX/Niagara/Source/NiagaraEditor/Public/ViewModels/Stack/NiagaraStackEmitterSettingsGroup.h
--- a/Engine/Plugins/FX/Niagara/Source/NiagaraEditor/Public/ViewModels/Stack/NiagaraStackEmitterSettingsGroup.h
+++ b/Engine/Plugins/FX/Niagara/Source/NiagaraEditor/Public/ViewModels/Stack/NiagaraStackEmitterSettingsGroup.h
@@ -46,7 +46,7 @@ private:
 	TWeakObjectPtr<UNiagaraEmitter> Emitter;
 
 	UPROPERTY()
-	UNiagaraStackObject* EmitterObject;
+	TObjectPtr<UNiagaraStackObject> EmitterObject;
 };
 
 UCLASS()
@@ -62,5 +62,5 @@ protected:
 
 private:
 	UPROPERTY()
-	UNiagaraStackEmitterPropertiesItem* PropertiesItem;
+	TObjectPtr<UNiagaraStackEmitterPropertiesItem> PropertiesItem;
 };
diff --git a/Engine/Plugins/FX/Niagara/Source/NiagaraEditor/Public/ViewModels/Stack/NiagaraStackEntry.h b/Engine/Plugins/FX/Niagara/Source/NiagaraEditor/Public/ViewModels/Stack/NiagaraStackEntry.h
--- a/Engine/Plugins/FX/Niagara/Source/NiagaraEditor/Public/ViewModels/Stack/NiagaraStackEntry.h
+++ b/Engine/Plugins/FX/Niagara/Source/NiagaraEditor/Public/ViewModels/Stack/NiagaraStackEntry.h
@@ -469,10 +469,10 @@ private:
 	TArray<FOnFilterChild> ChildFilters;
 
 	UPROPERTY()
-	TArray<UNiagaraStackEntry*> Children;
+	TArray<TObjectPtr<UNiagaraStackEntry>> Children;
 
 	UPROPERTY()
-	TArray<UNiagaraStackErrorItem*> ErrorChildren;
+	TArray<TObjectPtr<UNiagaraStackErrorItem>> ErrorChildren;
 
 	mutable TOptional<bool> bIsExpandedCache;
 
diff --git a/Engine/Plugins/FX/Niagara/Source/NiagaraEditor/Public/ViewModels/Stack/NiagaraStackEventScriptItemGroup.h b/Engine/Plugins/FX/Niagara/Source/NiagaraEditor/Public/ViewModels/Stack/NiagaraStackEventScriptItemGroup.h
--- a/Engine/Plugins/FX/Niagara/Source/NiagaraEditor/Public/ViewModels/Stack/NiagaraStackEventScriptItemGroup.h
+++ b/Engine/Plugins/FX/Niagara/Source/NiagaraEditor/Public/ViewModels/Stack/NiagaraStackEventScriptItemGroup.h
@@ -51,7 +51,7 @@ private:
 	TWeakObjectPtr<UNiagaraEmitter> Emitter;
 
 	UPROPERTY()
-	UNiagaraStackObject* EmitterObject;
+	TObjectPtr<UNiagaraStackObject> EmitterObject;
 };
 
 UCLASS()
@@ -93,5 +93,5 @@ private:
 	FGuid EventSourceEmitterId;
 
 	UPROPERTY()
-	UNiagaraStackEventHandlerPropertiesItem* EventHandlerProperties;
+	TObjectPtr<UNiagaraStackEventHandlerPropertiesItem> EventHandlerProperties;
 };
diff --git a/Engine/Plugins/FX/Niagara/Source/NiagaraEditor/Public/ViewModels/Stack/NiagaraStackItem.h b/Engine/Plugins/FX/Niagara/Source/NiagaraEditor/Public/ViewModels/Stack/NiagaraStackItem.h
--- a/Engine/Plugins/FX/Niagara/Source/NiagaraEditor/Public/ViewModels/Stack/NiagaraStackItem.h
+++ b/Engine/Plugins/FX/Niagara/Source/NiagaraEditor/Public/ViewModels/Stack/NiagaraStackItem.h
@@ -65,7 +65,7 @@ protected:
 
 private:
 	UPROPERTY()
-	UNiagaraStackItemFooter* ItemFooter;
+	TObjectPtr<UNiagaraStackItemFooter> ItemFooter;
 };
 
 UCLASS()
diff --git a/Engine/Plugins/FX/Niagara/Source/NiagaraEditor/Public/ViewModels/Stack/NiagaraStackItemGroup.h b/Engine/Plugins/FX/Niagara/Source/NiagaraEditor/Public/ViewModels/Stack/NiagaraStackItemGroup.h
--- a/Engine/Plugins/FX/Niagara/Source/NiagaraEditor/Public/ViewModels/Stack/NiagaraStackItemGroup.h
+++ b/Engine/Plugins/FX/Niagara/Source/NiagaraEditor/Public/ViewModels/Stack/NiagaraStackItemGroup.h
@@ -46,7 +46,7 @@ private:
 
 private:
 	UPROPERTY()
-	UNiagaraStackItemGroupFooter* GroupFooter;
+	TObjectPtr<UNiagaraStackItemGroupFooter> GroupFooter;
 
 	INiagaraStackItemGroupAddUtilities* AddUtilities;
 
diff --git a/Engine/Plugins/FX/Niagara/Source/NiagaraEditor/Public/ViewModels/Stack/NiagaraStackModuleItem.h b/Engine/Plugins/FX/Niagara/Source/NiagaraEditor/Public/ViewModels/Stack/NiagaraStackModuleItem.h
--- a/Engine/Plugins/FX/Niagara/Source/NiagaraEditor/Public/ViewModels/Stack/NiagaraStackModuleItem.h
+++ b/Engine/Plugins/FX/Niagara/Source/NiagaraEditor/Public/ViewModels/Stack/NiagaraStackModuleItem.h
@@ -131,13 +131,13 @@ private:
 	bool bCanRefresh;
 
 	UPROPERTY()
-	UNiagaraStackModuleItemLinkedInputCollection* LinkedInputCollection;
+	TObjectPtr<UNiagaraStackModuleItemLinkedInputCollection> LinkedInputCollection;
 
 	UPROPERTY()
-	UNiagaraStackFunctionInputCollection* InputCollection;
+	TObjectPtr<UNiagaraStackFunctionInputCollection> InputCollection;
 
 	UPROPERTY()
-	UNiagaraStackModuleItemOutputCollection* OutputCollection;
+	TObjectPtr<UNiagaraStackModuleItemOutputCollection> OutputCollection;
 
 	INiagaraStackItemGroupAddUtilities* GroupAddUtilities;
 
diff --git a/Engine/Plugins/FX/Niagara/Source/NiagaraEditor/Public/ViewModels/Stack/NiagaraStackParameterStoreEntry.h b/Engine/Plugins/FX/Niagara/Source/NiagaraEditor/Public/ViewModels/Stack/NiagaraStackParameterStoreEntry.h
--- a/Engine/Plugins/FX/Niagara/Source/NiagaraEditor/Public/ViewModels/Stack/NiagaraStackParameterStoreEntry.h
+++ b/Engine/Plugins/FX/Niagara/Source/NiagaraEditor/Public/ViewModels/Stack/NiagaraStackParameterStoreEntry.h
@@ -140,5 +140,5 @@ private:
 
 	/** The stack entry for displaying the value object. */
 	UPROPERTY()
-	UNiagaraStackObject* ValueObjectEntry;
+	TObjectPtr<UNiagaraStackObject> ValueObjectEntry;
 };
diff --git a/Engine/Plugins/FX/Niagara/Source/NiagaraEditor/Public/ViewModels/Stack/NiagaraStackRendererItem.h b/Engine/Plugins/FX/Niagara/Source/NiagaraEditor/Public/ViewModels/Stack/NiagaraStackRendererItem.h
--- a/Engine/Plugins/FX/Niagara/Source/NiagaraEditor/Public/ViewModels/Stack/NiagaraStackRendererItem.h
+++ b/Engine/Plugins/FX/Niagara/Source/NiagaraEditor/Public/ViewModels/Stack/NiagaraStackRendererItem.h
@@ -84,5 +84,5 @@ private:
 	TArray<FNiagaraVariable> MissingAttributes;
 
 	UPROPERTY()
-	UNiagaraStackObject* RendererObject;
+	TObjectPtr<UNiagaraStackObject> RendererObject;
 };
diff --git a/Engine/Plugins/FX/Niagara/Source/NiagaraEditor/Public/ViewModels/Stack/NiagaraStackRoot.h b/Engine/Plugins/FX/Niagara/Source/NiagaraEditor/Public/ViewModels/Stack/NiagaraStackRoot.h
--- a/Engine/Plugins/FX/Niagara/Source/NiagaraEditor/Public/ViewModels/Stack/NiagaraStackRoot.h
+++ b/Engine/Plugins/FX/Niagara/Source/NiagaraEditor/Public/ViewModels/Stack/NiagaraStackRoot.h
@@ -37,37 +37,37 @@ private:
 
 private:
 	UPROPERTY()
-	UNiagaraStackSystemSettingsGroup* SystemSettingsGroup;
+	TObjectPtr<UNiagaraStackSystemSettingsGroup> SystemSettingsGroup;
 
 	UPROPERTY()
-	UNiagaraStackScriptItemGroup* SystemSpawnGroup;
+	TObjectPtr<UNiagaraStackScriptItemGroup> SystemSpawnGroup;
 
 	UPROPERTY()
-	UNiagaraStackScriptItemGroup* SystemUpdateGroup;
+	TObjectPtr<UNiagaraStackScriptItemGroup> SystemUpdateGroup;
 
 	UPROPERTY()
-	UNiagaraStackEmitterSettingsGroup* EmitterSettingsGroup;
+	TObjectPtr<UNiagaraStackEmitterSettingsGroup> EmitterSettingsGroup;
 
 	UPROPERTY()
-	UNiagaraStackScriptItemGroup* EmitterSpawnGroup;
+	TObjectPtr<UNiagaraStackScriptItemGroup> EmitterSpawnGroup;
 
 	UPROPERTY()
-	UNiagaraStackScriptItemGroup* EmitterUpdateGroup;
+	TObjectPtr<UNiagaraStackScriptItemGroup> EmitterUpdateGroup;
 
 	UPROPERTY()
-	UNiagaraStackScriptItemGroup* ParticleSpawnGroup;
+	TObjectPtr<UNiagaraStackScriptItemGroup> ParticleSpawnGroup;
 
 	UPROPERTY()
-	UNiagaraStackScriptItemGroup* ParticleUpdateGroup;
+	TObjectPtr<UNiagaraStackScriptItemGroup> ParticleUpdateGroup;
 
 	UPROPERTY()
-	UNiagaraStackEventHandlerGroup* AddEventHandlerGroup;
+	TObjectPtr<UNiagaraStackEventHandlerGroup> AddEventHandlerGroup;
 
 	UPROPERTY()
-	UNiagaraStackSimulationStagesGroup* AddSimulationStageGroup;
+	TObjectPtr<UNiagaraStackSimulationStagesGroup> AddSimulationStageGroup;
 
 	UPROPERTY()
-	UNiagaraStackRenderItemGroup* RenderGroup;
+	TObjectPtr<UNiagaraStackRenderItemGroup> RenderGroup;
 
 	bool bIncludeSystemInformation;
 	bool bIncludeEmitterInformation;
diff --git a/Engine/Plugins/FX/Niagara/Source/NiagaraEditor/Public/ViewModels/Stack/NiagaraStackSimulationStageGroup.h b/Engine/Plugins/FX/Niagara/Source/NiagaraEditor/Public/ViewModels/Stack/NiagaraStackSimulationStageGroup.h
--- a/Engine/Plugins/FX/Niagara/Source/NiagaraEditor/Public/ViewModels/Stack/NiagaraStackSimulationStageGroup.h
+++ b/Engine/Plugins/FX/Niagara/Source/NiagaraEditor/Public/ViewModels/Stack/NiagaraStackSimulationStageGroup.h
@@ -45,7 +45,7 @@ private:
 	mutable TOptional<bool> bCanResetToBaseCache;
 
 	UPROPERTY()
-	UNiagaraStackObject* SimulationStageObject;
+	TObjectPtr<UNiagaraStackObject> SimulationStageObject;
 };
 
 UCLASS()
@@ -99,5 +99,5 @@ private:
 	FOnModifiedSimulationStages OnModifiedSimulationStagesDelegate;
 
 	UPROPERTY()
-	UNiagaraStackSimulationStagePropertiesItem* SimulationStageProperties;
+	TObjectPtr<UNiagaraStackSimulationStagePropertiesItem> SimulationStageProperties;
 };
diff --git a/Engine/Plugins/FX/Niagara/Source/NiagaraEditor/Public/ViewModels/Stack/NiagaraStackSystemPropertiesItem.h b/Engine/Plugins/FX/Niagara/Source/NiagaraEditor/Public/ViewModels/Stack/NiagaraStackSystemPropertiesItem.h
--- a/Engine/Plugins/FX/Niagara/Source/NiagaraEditor/Public/ViewModels/Stack/NiagaraStackSystemPropertiesItem.h
+++ b/Engine/Plugins/FX/Niagara/Source/NiagaraEditor/Public/ViewModels/Stack/NiagaraStackSystemPropertiesItem.h
@@ -34,6 +34,6 @@ private:
 	TWeakObjectPtr<UNiagaraSystem> System;
 
 	UPROPERTY()
-	UNiagaraStackObject* SystemObject;
+	TObjectPtr<UNiagaraStackObject> SystemObject;
 
 };
diff --git a/Engine/Plugins/FX/Niagara/Source/NiagaraEditor/Public/ViewModels/Stack/NiagaraStackViewModel.h b/Engine/Plugins/FX/Niagara/Source/NiagaraEditor/Public/ViewModels/Stack/NiagaraStackViewModel.h
--- a/Engine/Plugins/FX/Niagara/Source/NiagaraEditor/Public/ViewModels/Stack/NiagaraStackViewModel.h
+++ b/Engine/Plugins/FX/Niagara/Source/NiagaraEditor/Public/ViewModels/Stack/NiagaraStackViewModel.h
@@ -178,7 +178,7 @@ private:
 	TArray<UNiagaraStackEntry*> RootEntries;
 
 	UPROPERTY()
-	UNiagaraStackEntry* RootEntry;
+	TObjectPtr<UNiagaraStackEntry> RootEntry;
 
 	bool bExternalRootEntry;
 
diff --git a/Engine/Plugins/Lumin/MagicLeap/Source/MagicLeap/Public/MagicLeapMeshTrackerComponent.h b/Engine/Plugins/Lumin/MagicLeap/Source/MagicLeap/Public/MagicLeapMeshTrackerComponent.h
--- a/Engine/Plugins/Lumin/MagicLeap/Source/MagicLeap/Public/MagicLeapMeshTrackerComponent.h
+++ b/Engine/Plugins/Lumin/MagicLeap/Source/MagicLeap/Public/MagicLeapMeshTrackerComponent.h
@@ -87,7 +87,7 @@ public:
 
 	/** Bounding box for the mesh scan. The mesh will be scanned for only within this box. */
 	UPROPERTY(EditAnywhere, BlueprintReadWrite, Category = "Meshing|MagicLeap")
-	class UBoxComponent* BoundingVolume;
+	TObjectPtr<class UBoxComponent> BoundingVolume;
 
 	/** Meshing LOD. */
 	UPROPERTY(EditAnywhere, BlueprintReadWrite, Category = "Meshing|MagicLeap")
@@ -145,7 +145,7 @@ public:
 
 	/** MRMeshComponent can render and provide collision based on the Mesh data. */
 	UPROPERTY(transient)
-	class UMRMeshComponent* MRMesh;
+	TObjectPtr<class UMRMeshComponent> MRMesh;
 
 	/** Specifies the number of MR Mesh bricks to create per frame. This is to
 	    amortize the cost of brick creation across frames rather than creating
diff --git a/Engine/Plugins/Lumin/MagicLeap/Source/MagicLeapAR/Private/LuminAROrigin.h b/Engine/Plugins/Lumin/MagicLeap/Source/MagicLeapAR/Private/LuminAROrigin.h
--- a/Engine/Plugins/Lumin/MagicLeap/Source/MagicLeapAR/Private/LuminAROrigin.h
+++ b/Engine/Plugins/Lumin/MagicLeap/Source/MagicLeapAR/Private/LuminAROrigin.h
@@ -54,13 +54,13 @@ private:
 
 private:
 	UPROPERTY()
-	class UMRMeshComponent* MRMeshComponent;
+	TObjectPtr<class UMRMeshComponent> MRMeshComponent;
 
 	UPROPERTY()
-	class UMaterialInterface* PlaneSurfaceMaterial;
+	TObjectPtr<class UMaterialInterface> PlaneSurfaceMaterial;
 
 	UPROPERTY()
-	class UMaterialInterface* WireframeMaterial;
+	TObjectPtr<class UMaterialInterface> WireframeMaterial;
 
 	TMap<FGuid, IMRMesh::FBrickId> PlaneIdToBrickId;
 	IMRMesh::FBrickId NextBrickId = 0;
diff --git a/Engine/Plugins/Lumin/MagicLeap/Source/MagicLeapImageTracker/Public/MagicLeapImageTrackerComponent.h b/Engine/Plugins/Lumin/MagicLeap/Source/MagicLeapImageTracker/Public/MagicLeapImageTrackerComponent.h
--- a/Engine/Plugins/Lumin/MagicLeap/Source/MagicLeapImageTracker/Public/MagicLeapImageTrackerComponent.h
+++ b/Engine/Plugins/Lumin/MagicLeap/Source/MagicLeapImageTracker/Public/MagicLeapImageTrackerComponent.h
@@ -54,7 +54,7 @@ public:
 
 	/** The texture that will be tracked by this image tracker instance. */
 	UPROPERTY(EditAnywhere, BlueprintReadWrite, Category = "ImageTracking|MagicLeap")
-	UTexture2D* TargetImageTexture;
+	TObjectPtr<UTexture2D> TargetImageTexture;
 
 	/**
 	  The name of the target.
diff --git a/Engine/Plugins/Lumin/MagicLeap/Source/MagicLeapImageTracker/Public/MagicLeapImageTrackerTypes.h b/Engine/Plugins/Lumin/MagicLeap/Source/MagicLeapImageTracker/Public/MagicLeapImageTrackerTypes.h
--- a/Engine/Plugins/Lumin/MagicLeap/Source/MagicLeapImageTracker/Public/MagicLeapImageTrackerTypes.h
+++ b/Engine/Plugins/Lumin/MagicLeap/Source/MagicLeapImageTracker/Public/MagicLeapImageTrackerTypes.h
@@ -58,7 +58,7 @@ struct MAGICLEAPIMAGETRACKER_API FMagicLeapImageTargetSettings
 
 public:
 	UPROPERTY(EditAnywhere, BlueprintReadWrite, Category = "ImageTracking|MagicLeap")
-	UTexture2D* ImageTexture;
+	TObjectPtr<UTexture2D> ImageTexture;
 
 	UPROPERTY(EditAnywhere, BlueprintReadWrite, Category = "ImageTracking|MagicLeap")
 	FString Name;
diff --git a/Engine/Plugins/Lumin/MagicLeap/Source/MagicLeapPlanes/Source/Public/MagicLeapPlanesComponent.h b/Engine/Plugins/Lumin/MagicLeap/Source/MagicLeapPlanes/Source/Public/MagicLeapPlanesComponent.h
--- a/Engine/Plugins/Lumin/MagicLeap/Source/MagicLeapPlanes/Source/Public/MagicLeapPlanesComponent.h
+++ b/Engine/Plugins/Lumin/MagicLeap/Source/MagicLeapPlanes/Source/Public/MagicLeapPlanesComponent.h
@@ -34,7 +34,7 @@ public:
 
 	/** Bounding box for searching planes in. */
 	UPROPERTY(EditAnywhere, BlueprintReadWrite, Category = "Planes | MagicLeap")
-	class UBoxComponent* SearchVolume;
+	TObjectPtr<class UBoxComponent> SearchVolume;
 
 	/** The maximum number of planes that should be returned in the result. */
 	UPROPERTY(EditAnywhere, BlueprintReadWrite, Category = "Planes | MagicLeap", meta = (ClampMin = 0))
diff --git a/Engine/Plugins/Lumin/MagicLeap/Source/MagicLeapPlanes/Source/Public/MagicLeapPlanesTypes.h b/Engine/Plugins/Lumin/MagicLeap/Source/MagicLeapPlanes/Source/Public/MagicLeapPlanesTypes.h
--- a/Engine/Plugins/Lumin/MagicLeap/Source/MagicLeapPlanes/Source/Public/MagicLeapPlanesTypes.h
+++ b/Engine/Plugins/Lumin/MagicLeap/Source/MagicLeapPlanes/Source/Public/MagicLeapPlanesTypes.h
@@ -109,7 +109,7 @@ struct MAGICLEAPPLANES_API FMagicLeapPlanesQuery
 
 	/** DEPRECATED. Use individual fields for setting search volume position, orientation and extents. */
 	UPROPERTY(EditAnywhere, BlueprintReadWrite, Category = "Planes|MagicLeap", meta = (DeprecatedProperty, DisplayName = "SearchVolume_DEPRECATED"))
-	class UBoxComponent* SearchVolume;
+	TObjectPtr<class UBoxComponent> SearchVolume;
 
 	/**
 		The maximum number of results that should be returned.  This is also the minimum expected size of the array of results
diff --git a/Engine/Plugins/Lumin/MagicLeapPassableWorld/Source/MagicLeapARPin/Public/Debug/MagicLeapARPinRenderer.h b/Engine/Plugins/Lumin/MagicLeapPassableWorld/Source/MagicLeapARPin/Public/Debug/MagicLeapARPinRenderer.h
--- a/Engine/Plugins/Lumin/MagicLeapPassableWorld/Source/MagicLeapARPin/Public/Debug/MagicLeapARPinRenderer.h
+++ b/Engine/Plugins/Lumin/MagicLeapPassableWorld/Source/MagicLeapARPin/Public/Debug/MagicLeapARPinRenderer.h
@@ -28,7 +28,7 @@ private:
 	void SetVisibilityOverride(const bool InVisibilityOverride);
 
 	UPROPERTY()
-	TMap<FGuid, AMagicLeapARPinInfoActorBase*> AllInfoActors;
+	TMap<FGuid, TObjectPtr<AMagicLeapARPinInfoActorBase>> AllInfoActors;
 
 	FDelegateHandle DelegateHandle;
 
diff --git a/Engine/Plugins/Lumin/MagicLeapPassableWorld/Source/MagicLeapARPin/Public/MagicLeapARPinComponent.h b/Engine/Plugins/Lumin/MagicLeapPassableWorld/Source/MagicLeapARPin/Public/MagicLeapARPinComponent.h
--- a/Engine/Plugins/Lumin/MagicLeapPassableWorld/Source/MagicLeapARPin/Public/MagicLeapARPinComponent.h
+++ b/Engine/Plugins/Lumin/MagicLeapPassableWorld/Source/MagicLeapARPin/Public/MagicLeapARPinComponent.h
@@ -182,7 +182,7 @@ public:
 
 	/** Volume to search for an ARPin in. The position and scaled radius (in Unreal Units) of this sphere is used to look for an ARPin of type SearchPinTypes. */
 	UPROPERTY(BlueprintReadWrite, EditAnywhere, Category = "ContentPersistence|MagicLeap")
-	class USphereComponent* SearchVolume;
+	TObjectPtr<class USphereComponent> SearchVolume;
 
 	/**
 	 * Delegate used to notify the instigating blueprint that an entity (component or actor) has been successfully pinned to the real-world.
@@ -226,10 +226,10 @@ private:
 	FGuid PinnedCFUID;
 
 	UPROPERTY()
-	USceneComponent* PinnedSceneComponent;
+	TObjectPtr<USceneComponent> PinnedSceneComponent;
 
 	UPROPERTY()
-	UMagicLeapARPinSaveGame* PinData;
+	TObjectPtr<UMagicLeapARPinSaveGame> PinData;
 
 	FTransform OldComponentWorldTransform;
 	FTransform OldCFUIDTransform;
diff --git a/Engine/Plugins/Lumin/MagicLeapPassableWorld/Source/MagicLeapSharedWorld/Public/MagicLeapSharedWorldGameMode.h b/Engine/Plugins/Lumin/MagicLeapPassableWorld/Source/MagicLeapSharedWorld/Public/MagicLeapSharedWorldGameMode.h
--- a/Engine/Plugins/Lumin/MagicLeapPassableWorld/Source/MagicLeapSharedWorld/Public/MagicLeapSharedWorldGameMode.h
+++ b/Engine/Plugins/Lumin/MagicLeapPassableWorld/Source/MagicLeapSharedWorld/Public/MagicLeapSharedWorldGameMode.h
@@ -90,5 +90,5 @@ protected:
 	TMap<FGuid, TMap<AMagicLeapSharedWorldPlayerController*, FMagicLeapARPinState>> PinToOwnersToStates;
 
 	UPROPERTY(BlueprintReadWrite, Category="AR Shared World|Magic Leap")
-	AMagicLeapSharedWorldPlayerController* ChosenOne;
+	TObjectPtr<AMagicLeapSharedWorldPlayerController> ChosenOne;
 };
diff --git a/Engine/Plugins/Media/MediaCompositing/Source/MediaCompositing/Private/MovieScene/MovieSceneMediaPlayerPropertyTemplate.h b/Engine/Plugins/Media/MediaCompositing/Source/MediaCompositing/Private/MovieScene/MovieSceneMediaPlayerPropertyTemplate.h
--- a/Engine/Plugins/Media/MediaCompositing/Source/MediaCompositing/Private/MovieScene/MovieSceneMediaPlayerPropertyTemplate.h
+++ b/Engine/Plugins/Media/MediaCompositing/Source/MediaCompositing/Private/MovieScene/MovieSceneMediaPlayerPropertyTemplate.h
@@ -30,7 +30,7 @@ public:
 private:
 
 	UPROPERTY()
-	UMediaSource* MediaSource;
+	TObjectPtr<UMediaSource> MediaSource;
 
 	UPROPERTY()
 	FFrameNumber SectionStartFrame;
diff --git a/Engine/Plugins/Media/MediaCompositing/Source/MediaCompositing/Private/MovieScene/MovieSceneMediaTemplate.h b/Engine/Plugins/Media/MediaCompositing/Source/MediaCompositing/Private/MovieScene/MovieSceneMediaTemplate.h
--- a/Engine/Plugins/Media/MediaCompositing/Source/MediaCompositing/Private/MovieScene/MovieSceneMediaTemplate.h
+++ b/Engine/Plugins/Media/MediaCompositing/Source/MediaCompositing/Private/MovieScene/MovieSceneMediaTemplate.h
@@ -21,16 +21,16 @@ struct FMovieSceneMediaSectionParams
 	GENERATED_BODY()
 
 	UPROPERTY()
-	UMediaSoundComponent* MediaSoundComponent;
+	TObjectPtr<UMediaSoundComponent> MediaSoundComponent;
 
 	UPROPERTY()
-	UMediaSource* MediaSource;
+	TObjectPtr<UMediaSource> MediaSource;
 
 	UPROPERTY()
-	UMediaTexture* MediaTexture;
+	TObjectPtr<UMediaTexture> MediaTexture;
 
 	UPROPERTY()
-	UMediaPlayer* MediaPlayer;
+	TObjectPtr<UMediaPlayer> MediaPlayer;
 
 	UPROPERTY()
 	FFrameNumber SectionStartFrame;
diff --git a/Engine/Plugins/Media/MediaCompositing/Source/MediaCompositing/Public/MovieSceneMediaPlayerPropertySection.h b/Engine/Plugins/Media/MediaCompositing/Source/MediaCompositing/Public/MovieSceneMediaPlayerPropertySection.h
--- a/Engine/Plugins/Media/MediaCompositing/Source/MediaCompositing/Public/MovieSceneMediaPlayerPropertySection.h
+++ b/Engine/Plugins/Media/MediaCompositing/Source/MediaCompositing/Public/MovieSceneMediaPlayerPropertySection.h
@@ -20,7 +20,7 @@ public:
 
 	/** The source to play with this video track. */
 	UPROPERTY(EditAnywhere, BlueprintReadWrite, Category="Media")
-	UMediaSource* MediaSource;
+	TObjectPtr<UMediaSource> MediaSource;
 
 	/** Whether to loop this video. */
 	UPROPERTY(EditAnywhere, BlueprintReadWrite, Category="Media")
diff --git a/Engine/Plugins/Media/MediaCompositing/Source/MediaCompositing/Public/MovieSceneMediaSection.h b/Engine/Plugins/Media/MediaCompositing/Source/MediaCompositing/Public/MovieSceneMediaSection.h
--- a/Engine/Plugins/Media/MediaCompositing/Source/MediaCompositing/Public/MovieSceneMediaSection.h
+++ b/Engine/Plugins/Media/MediaCompositing/Source/MediaCompositing/Public/MovieSceneMediaSection.h
@@ -27,7 +27,7 @@ public:
 public:
 	/** The source to play with this video track. */
 	UPROPERTY(EditAnywhere, BlueprintReadWrite, Category="Media")
-	UMediaSource* MediaSource;
+	TObjectPtr<UMediaSource> MediaSource;
 
 	/** Should the media player be set to loop? This can be helpful for media formats that can use this information (such as exr sequences) to pre-cache the starting data when nearing the end of playback. Does not cause the media to continue playing after the end of the section is reached. */
 	UPROPERTY(EditAnywhere, BlueprintReadWrite, Category = "Media")
@@ -39,11 +39,11 @@ public:
 
 	/** The media texture that receives the track's video output. */
 	UPROPERTY(EditAnywhere, BlueprintReadWrite, Category="Media", meta = (EditCondition = "!bUseExternalMediaPlayer"))
-	UMediaTexture* MediaTexture;
+	TObjectPtr<UMediaTexture> MediaTexture;
 
 	/** The media sound component that receives the track's audio output. */
 	UPROPERTY(EditAnywhere, BlueprintReadWrite, Category="Media")
-	UMediaSoundComponent* MediaSoundComponent;
+	TObjectPtr<UMediaSoundComponent> MediaSoundComponent;
 
 	/** If true, this track will control a previously created media player instead of automatically creating one. */
 	UPROPERTY(EditAnywhere, BlueprintReadWrite, Category="Media", AdvancedDisplay)
@@ -51,7 +51,7 @@ public:
 
 	/** The external media player this track should control. */
 	UPROPERTY(EditAnywhere, BlueprintReadWrite, Category="Media", meta = (EditCondition = "bUseExternalMediaPlayer"), AdvancedDisplay)
-	UMediaPlayer* ExternalMediaPlayer;
+	TObjectPtr<UMediaPlayer> ExternalMediaPlayer;
 
 public:
 
diff --git a/Engine/Plugins/Media/MediaCompositing/Source/MediaCompositing/Public/MovieSceneMediaTrack.h b/Engine/Plugins/Media/MediaCompositing/Source/MediaCompositing/Public/MovieSceneMediaTrack.h
--- a/Engine/Plugins/Media/MediaCompositing/Source/MediaCompositing/Public/MovieSceneMediaTrack.h
+++ b/Engine/Plugins/Media/MediaCompositing/Source/MediaCompositing/Public/MovieSceneMediaTrack.h
@@ -58,5 +58,5 @@ private:
 
 	/** List of all master media sections. */
 	UPROPERTY()
-	TArray<UMovieSceneSection*> MediaSections;
+	TArray<TObjectPtr<UMovieSceneSection>> MediaSections;
 };
diff --git a/Engine/Plugins/Media/MediaPlayerEditor/Source/MediaPlayerEditor/Private/Shared/MediaPlayerEditorMediaContext.h b/Engine/Plugins/Media/MediaPlayerEditor/Source/MediaPlayerEditor/Private/Shared/MediaPlayerEditorMediaContext.h
--- a/Engine/Plugins/Media/MediaPlayerEditor/Source/MediaPlayerEditor/Private/Shared/MediaPlayerEditorMediaContext.h
+++ b/Engine/Plugins/Media/MediaPlayerEditor/Source/MediaPlayerEditor/Private/Shared/MediaPlayerEditorMediaContext.h
@@ -25,7 +25,7 @@ public:
 	}
 
 	UPROPERTY(BlueprintReadWrite, Category = "Editor Menus")
-	UObject* SelectedAsset;
+	TObjectPtr<UObject> SelectedAsset;
 
 	UPROPERTY(BlueprintReadWrite, Category = "Editor Menus")
 	FName StyleSetName;
diff --git a/Engine/Plugins/MeshPainting/Source/MeshPaintEditorMode/Private/MeshPaintMode.h b/Engine/Plugins/MeshPainting/Source/MeshPaintEditorMode/Private/MeshPaintMode.h
--- a/Engine/Plugins/MeshPainting/Source/MeshPaintEditorMode/Private/MeshPaintMode.h
+++ b/Engine/Plugins/MeshPainting/Source/MeshPaintEditorMode/Private/MeshPaintMode.h
@@ -109,7 +109,7 @@ protected:
 
 protected:
 	UPROPERTY(Transient)
-	UMeshPaintModeSettings* ModeSettings;
+	TObjectPtr<UMeshPaintModeSettings> ModeSettings;
 
 
 	// End vertex paint state
diff --git a/Engine/Plugins/MeshPainting/Source/MeshPaintingToolset/Public/MeshPaintInteractions.h b/Engine/Plugins/MeshPainting/Source/MeshPaintingToolset/Public/MeshPaintInteractions.h
--- a/Engine/Plugins/MeshPainting/Source/MeshPaintingToolset/Public/MeshPaintInteractions.h
+++ b/Engine/Plugins/MeshPainting/Source/MeshPaintingToolset/Public/MeshPaintInteractions.h
@@ -44,8 +44,8 @@ protected:
 
 protected:
 	UPROPERTY(Transient)
-	TArray<UMeshComponent*> CachedClickedComponents;
+	TArray<TObjectPtr<UMeshComponent>> CachedClickedComponents;
 	UPROPERTY(Transient)
-	TArray<AActor*> CachedClickedActors;
+	TArray<TObjectPtr<AActor>> CachedClickedActors;
 	bool bAddToSelectionSet;
 };
diff --git a/Engine/Plugins/MeshPainting/Source/MeshPaintingToolset/Public/MeshPaintingToolsetTypes.h b/Engine/Plugins/MeshPainting/Source/MeshPaintingToolset/Public/MeshPaintingToolsetTypes.h
--- a/Engine/Plugins/MeshPainting/Source/MeshPaintingToolset/Public/MeshPaintingToolsetTypes.h
+++ b/Engine/Plugins/MeshPainting/Source/MeshPaintingToolset/Public/MeshPaintingToolsetTypes.h
@@ -121,24 +121,24 @@ struct FPaintTexture2DData
 	GENERATED_BODY()
 	/** The original texture that we're painting */
 	UPROPERTY(Transient)
-	UTexture2D* PaintingTexture2D;
+	TObjectPtr<UTexture2D> PaintingTexture2D;
 	bool bIsPaintingTexture2DModified;
 
 	/** A copy of the original texture we're painting, used for restoration. */
 	UPROPERTY(Transient)
-	UTexture2D* PaintingTexture2DDuplicate;
+	TObjectPtr<UTexture2D> PaintingTexture2DDuplicate;
 
 	/** Render target texture for painting */
 	UPROPERTY(Transient)
-	UTextureRenderTarget2D* PaintRenderTargetTexture;
+	TObjectPtr<UTextureRenderTarget2D> PaintRenderTargetTexture;
 
 	/** Render target texture used as an input while painting that contains a clone of the original image */
 	UPROPERTY(Transient)
-	UTextureRenderTarget2D* CloneRenderTargetTexture;
+	TObjectPtr<UTextureRenderTarget2D> CloneRenderTargetTexture;
 
 	/** List of materials we are painting on */
 	UPROPERTY(Transient)
-	TArray< UMaterialInterface* > PaintingMaterials;
+	TArray< TObjectPtr<UMaterialInterface> > PaintingMaterials;
 
 	/** Default ctor */
 	FPaintTexture2DData() :
diff --git a/Engine/Plugins/MeshPainting/Source/MeshPaintingToolset/Public/MeshSelect.h b/Engine/Plugins/MeshPainting/Source/MeshPaintingToolset/Public/MeshSelect.h
--- a/Engine/Plugins/MeshPainting/Source/MeshPaintingToolset/Public/MeshSelect.h
+++ b/Engine/Plugins/MeshPainting/Source/MeshPaintingToolset/Public/MeshSelect.h
@@ -81,7 +81,7 @@ protected:
 
 
 	UPROPERTY(Transient)
-	UMeshPaintSelectionMechanic* SelectionMechanic;
+	TObjectPtr<UMeshPaintSelectionMechanic> SelectionMechanic;
 
 };
 
diff --git a/Engine/Plugins/MeshPainting/Source/MeshPaintingToolset/Public/MeshTexturePaintingTool.h b/Engine/Plugins/MeshPainting/Source/MeshPaintingToolset/Public/MeshTexturePaintingTool.h
--- a/Engine/Plugins/MeshPainting/Source/MeshPaintingToolset/Public/MeshTexturePaintingTool.h
+++ b/Engine/Plugins/MeshPainting/Source/MeshPaintingToolset/Public/MeshTexturePaintingTool.h
@@ -82,7 +82,7 @@ public:
 
 	/** Texture to which Painting should be Applied */
 	UPROPERTY(EditAnywhere, Category = TexturePainting, meta = (DisplayThumbnail = "true", TransientToolProperty))
-	UTexture2D* PaintTexture;
+	TObjectPtr<UTexture2D> PaintTexture;
 
 	/** Enables "Flow" painting where paint is continually applied from the brush every tick */
 	UPROPERTY(EditAnywhere, Category = Brush, meta = (DisplayName = "Enable Brush Flow"))
@@ -164,10 +164,10 @@ private:
 	
 private:
 	UPROPERTY(Transient)
-	UMeshTexturePaintingToolProperties* TextureProperties;
+	TObjectPtr<UMeshTexturePaintingToolProperties> TextureProperties;
 
 	UPROPERTY(Transient)
-	TArray<const UTexture*> Textures;
+	TArray<TObjectPtr<const UTexture>> Textures;
 
 	/** Flag for whether or not we are currently painting */
 	bool bArePainting;
@@ -186,30 +186,30 @@ private:
 
 	/** Temporary render target used to draw incremental paint to */
 	UPROPERTY(Transient)
-	UTextureRenderTarget2D* BrushRenderTargetTexture;
+	TObjectPtr<UTextureRenderTarget2D> BrushRenderTargetTexture;
 
 	/** Temporary render target used to store a mask of the affected paint region, updated every time we add incremental texture paint */
 	UPROPERTY(Transient)
-	UTextureRenderTarget2D* BrushMaskRenderTargetTexture;
+	TObjectPtr<UTextureRenderTarget2D> BrushMaskRenderTargetTexture;
 
 	/** Temporary render target used to store generated mask for texture seams, we create this by projecting object triangles into texture space using the selected UV channel */
 	UPROPERTY(Transient)
-	UTextureRenderTarget2D* SeamMaskRenderTargetTexture;
+	TObjectPtr<UTextureRenderTarget2D> SeamMaskRenderTargetTexture;
 
 	/** Stores data associated with our paint target textures */
 	UPROPERTY(Transient)
-	TMap< UTexture2D*, FPaintTexture2DData > PaintTargetData;
+	TMap< TObjectPtr<UTexture2D>, FPaintTexture2DData > PaintTargetData;
 
 	/** Texture paint: Will hold a list of texture items that we can paint on */
 	TArray<FTextureTargetListInfo> TexturePaintTargetList;
 
 	/** Texture paint: The mesh components that we're currently painting */
 	UPROPERTY(Transient)
-	UMeshComponent* TexturePaintingCurrentMeshComponent;
+	TObjectPtr<UMeshComponent> TexturePaintingCurrentMeshComponent;
 
 	/** The original texture that we're painting */
 	UPROPERTY(Transient)
-	UTexture2D* PaintingTexture2D;
+	TObjectPtr<UTexture2D> PaintingTexture2D;
 
 	/** True if we need to generate a texture seam mask used for texture dilation */
 	bool bGenerateSeamMask;
diff --git a/Engine/Plugins/MeshPainting/Source/MeshPaintingToolset/Public/MeshVertexPaintingTool.h b/Engine/Plugins/MeshPainting/Source/MeshPaintingToolset/Public/MeshVertexPaintingTool.h
--- a/Engine/Plugins/MeshPainting/Source/MeshPaintingToolset/Public/MeshVertexPaintingTool.h
+++ b/Engine/Plugins/MeshPainting/Source/MeshPaintingToolset/Public/MeshVertexPaintingTool.h
@@ -223,14 +223,14 @@ protected:
 	bool bCachedClickRay;
 
 	UPROPERTY(Transient)
-	UMeshPaintSelectionMechanic* SelectionMechanic;
+	TObjectPtr<UMeshPaintSelectionMechanic> SelectionMechanic;
 
 private:
 	bool PaintInternal(const TArrayView<TPair<FVector, FVector>>& Rays, EMeshPaintModeAction PaintAction, float PaintStrength);
 	
 private:
 	UPROPERTY(Transient)
-	UMeshVertexPaintingToolProperties* VertexProperties;
+	TObjectPtr<UMeshVertexPaintingToolProperties> VertexProperties;
 
 	/** Flag for whether or not we are currently painting */
 	bool bArePainting;
@@ -268,7 +268,7 @@ protected:
 
 private:
 	UPROPERTY(Transient)
-	UMeshColorPaintingToolProperties* ColorProperties;
+	TObjectPtr<UMeshColorPaintingToolProperties> ColorProperties;
 
 	/** Current LOD index used for painting / forcing */
 	int32 CachedLODIndex;
@@ -292,5 +292,5 @@ protected:
 
 private:
 	UPROPERTY(Transient)
-	UMeshWeightPaintingToolProperties* WeightProperties;
+	TObjectPtr<UMeshWeightPaintingToolProperties> WeightProperties;
 };
diff --git a/Engine/Plugins/MovieScene/ActorSequence/Source/ActorSequence/Public/ActorSequence.h b/Engine/Plugins/MovieScene/ActorSequence/Source/ActorSequence/Public/ActorSequence.h
--- a/Engine/Plugins/MovieScene/ActorSequence/Source/ActorSequence/Public/ActorSequence.h
+++ b/Engine/Plugins/MovieScene/ActorSequence/Source/ActorSequence/Public/ActorSequence.h
@@ -45,7 +45,7 @@ private:
 	
 	/** Pointer to the movie scene that controls this animation. */
 	UPROPERTY(Instanced)
-	UMovieScene* MovieScene;
+	TObjectPtr<UMovieScene> MovieScene;
 
 	/** Collection of object references. */
 	UPROPERTY()
diff --git a/Engine/Plugins/MovieScene/ActorSequence/Source/ActorSequence/Public/ActorSequenceComponent.h b/Engine/Plugins/MovieScene/ActorSequence/Source/ActorSequence/Public/ActorSequenceComponent.h
--- a/Engine/Plugins/MovieScene/ActorSequence/Source/ActorSequence/Public/ActorSequenceComponent.h
+++ b/Engine/Plugins/MovieScene/ActorSequence/Source/ActorSequence/Public/ActorSequenceComponent.h
@@ -45,8 +45,8 @@ protected:
 
 	/** Embedded actor sequence data */
 	UPROPERTY(EditAnywhere, Instanced, Category=Animation)
-	UActorSequence* Sequence;
+	TObjectPtr<UActorSequence> Sequence;
 
 	UPROPERTY(transient, BlueprintReadOnly, Category=Animation)
-	UActorSequencePlayer* SequencePlayer;
+	TObjectPtr<UActorSequencePlayer> SequencePlayer;
 };
diff --git a/Engine/Plugins/MovieScene/LevelSequenceEditor/Source/LevelSequenceEditor/Private/LevelSequenceEditorBlueprintLibrary.h b/Engine/Plugins/MovieScene/LevelSequenceEditor/Source/LevelSequenceEditor/Private/LevelSequenceEditorBlueprintLibrary.h
--- a/Engine/Plugins/MovieScene/LevelSequenceEditor/Source/LevelSequenceEditor/Private/LevelSequenceEditorBlueprintLibrary.h
+++ b/Engine/Plugins/MovieScene/LevelSequenceEditor/Source/LevelSequenceEditor/Private/LevelSequenceEditorBlueprintLibrary.h
@@ -32,7 +32,7 @@ struct FSequencerChannelProxy
 	FName ChannelName;
 
 	UPROPERTY(BlueprintReadWrite, Category=Channel)
-	UMovieSceneSection* Section;
+	TObjectPtr<UMovieSceneSection> Section;
 };
 
 
diff --git a/Engine/Plugins/MovieScene/TemplateSequence/Source/TemplateSequence/Public/SequenceCameraShake.h b/Engine/Plugins/MovieScene/TemplateSequence/Source/TemplateSequence/Public/SequenceCameraShake.h
--- a/Engine/Plugins/MovieScene/TemplateSequence/Source/TemplateSequence/Public/SequenceCameraShake.h
+++ b/Engine/Plugins/MovieScene/TemplateSequence/Source/TemplateSequence/Public/SequenceCameraShake.h
@@ -104,7 +104,7 @@ public:
 
 	/** Source camera animation sequence to play. */
 	UPROPERTY(EditAnywhere, Category=CameraShake)
-	class UCameraAnimationSequence* Sequence;
+	TObjectPtr<class UCameraAnimationSequence> Sequence;
 
 	/** Scalar defining how fast to play the anim. */
 	UPROPERTY(EditAnywhere, Category=CameraShake, meta=(ClampMin="0.001"))
@@ -152,11 +152,11 @@ private:
 
 	/** The player we use to play the camera animation sequence */
 	UPROPERTY(Instanced, Transient)
-	USequenceCameraShakeSequencePlayer* Player;
+	TObjectPtr<USequenceCameraShakeSequencePlayer> Player;
 
 	/** Standin for the camera actor and components */
 	UPROPERTY(Instanced, Transient)
-	USequenceCameraShakeCameraStandIn* CameraStandIn;
+	TObjectPtr<USequenceCameraShakeCameraStandIn> CameraStandIn;
 };
 
 /**
@@ -237,11 +237,11 @@ private:
 
 	/** Bound object overrides */
 	UPROPERTY(transient)
-	UObject* BoundObjectOverride;
+	TObjectPtr<UObject> BoundObjectOverride;
 
 	/** The sequence to play back */
 	UPROPERTY(transient)
-	UMovieSceneSequence* Sequence;
+	TObjectPtr<UMovieSceneSequence> Sequence;
 
 	/** The evaluation template instance */
 	UPROPERTY(transient)
diff --git a/Engine/Plugins/MovieScene/TemplateSequence/Source/TemplateSequence/Public/TemplateSequenceActor.h b/Engine/Plugins/MovieScene/TemplateSequence/Source/TemplateSequence/Public/TemplateSequenceActor.h
--- a/Engine/Plugins/MovieScene/TemplateSequence/Source/TemplateSequence/Public/TemplateSequenceActor.h
+++ b/Engine/Plugins/MovieScene/TemplateSequence/Source/TemplateSequence/Public/TemplateSequenceActor.h
@@ -48,7 +48,7 @@ public:
 	FMovieSceneSequencePlaybackSettings PlaybackSettings;
 
 	UPROPERTY(Instanced, Transient, Replicated, BlueprintReadOnly, BlueprintGetter = GetSequencePlayer, Category = "Playback", meta = (ExposeFunctionCategories = "Game|Cinematic"))
-	UTemplateSequencePlayer* SequencePlayer;
+	TObjectPtr<UTemplateSequencePlayer> SequencePlayer;
 
 	UPROPERTY(EditAnywhere, BlueprintReadOnly, Category = "General", meta = (AllowedClasses = "TemplateSequence"))
 	FSoftObjectPath TemplateSequence;
diff --git a/Engine/Plugins/Online/OnlineSubsystem/Source/Public/NamedInterfaces.h b/Engine/Plugins/Online/OnlineSubsystem/Source/Public/NamedInterfaces.h
--- a/Engine/Plugins/Online/OnlineSubsystem/Source/Public/NamedInterfaces.h
+++ b/Engine/Plugins/Online/OnlineSubsystem/Source/Public/NamedInterfaces.h
@@ -28,7 +28,7 @@ struct FNamedInterface
 	FName InterfaceName;
 	/** The object to store at this location */
 	UPROPERTY()
-	UObject* InterfaceObject;
+	TObjectPtr<UObject> InterfaceObject;
 };
 
 /** Holds a name to class name mapping for adding the named interfaces automatically */
diff --git a/Engine/Plugins/Online/OnlineSubsystemUtils/Source/OnlineSubsystemUtils/Public/OnlineBeacon.h b/Engine/Plugins/Online/OnlineSubsystemUtils/Source/OnlineSubsystemUtils/Public/OnlineBeacon.h
--- a/Engine/Plugins/Online/OnlineSubsystemUtils/Source/OnlineSubsystemUtils/Public/OnlineBeacon.h
+++ b/Engine/Plugins/Online/OnlineSubsystemUtils/Source/OnlineSubsystemUtils/Public/OnlineBeacon.h
@@ -104,7 +104,7 @@ protected:
 
 	/** Net driver routing network traffic */
 	UPROPERTY()
-	UNetDriver* NetDriver;
+	TObjectPtr<UNetDriver> NetDriver;
 
 	/** State of beacon */
 	EBeaconState::Type BeaconState;
diff --git a/Engine/Plugins/Online/OnlineSubsystemUtils/Source/OnlineSubsystemUtils/Public/OnlineBeaconClient.h b/Engine/Plugins/Online/OnlineSubsystemUtils/Source/OnlineSubsystemUtils/Public/OnlineBeaconClient.h
--- a/Engine/Plugins/Online/OnlineSubsystemUtils/Source/OnlineSubsystemUtils/Public/OnlineBeaconClient.h
+++ b/Engine/Plugins/Online/OnlineSubsystemUtils/Source/OnlineSubsystemUtils/Public/OnlineBeaconClient.h
@@ -169,11 +169,11 @@ protected:
 
 	/** Owning beacon host of this beacon actor (server only) */
 	UPROPERTY()
-	AOnlineBeaconHostObject* BeaconOwner;
+	TObjectPtr<AOnlineBeaconHostObject> BeaconOwner;
 
 	/** Network connection associated with this beacon client instance */
 	UPROPERTY()
-	UNetConnection* BeaconConnection;
+	TObjectPtr<UNetConnection> BeaconConnection;
 
 	/** State of the connection */
 	UPROPERTY()
diff --git a/Engine/Plugins/Online/OnlineSubsystemUtils/Source/OnlineSubsystemUtils/Public/OnlineBeaconHost.h b/Engine/Plugins/Online/OnlineSubsystemUtils/Source/OnlineSubsystemUtils/Public/OnlineBeaconHost.h
--- a/Engine/Plugins/Online/OnlineSubsystemUtils/Source/OnlineSubsystemUtils/Public/OnlineBeaconHost.h
+++ b/Engine/Plugins/Online/OnlineSubsystemUtils/Source/OnlineSubsystemUtils/Public/OnlineBeaconHost.h
@@ -107,7 +107,7 @@ private:
 
 	/** List of all client beacon actors with active connections */
 	UPROPERTY()
-	TArray<AOnlineBeaconClient*> ClientActors;
+	TArray<TObjectPtr<AOnlineBeaconClient>> ClientActors;
 
 	/** Sends the welcome control message to the client. Used as a delegate if encryption is being negotiated. */
 	void SendWelcomeControlMessage(UNetConnection* Connection);
diff --git a/Engine/Plugins/Online/OnlineSubsystemUtils/Source/OnlineSubsystemUtils/Public/OnlineBeaconHostObject.h b/Engine/Plugins/Online/OnlineSubsystemUtils/Source/OnlineSubsystemUtils/Public/OnlineBeaconHostObject.h
--- a/Engine/Plugins/Online/OnlineSubsystemUtils/Source/OnlineSubsystemUtils/Public/OnlineBeaconHostObject.h
+++ b/Engine/Plugins/Online/OnlineSubsystemUtils/Source/OnlineSubsystemUtils/Public/OnlineBeaconHostObject.h
@@ -87,5 +87,5 @@ protected:
 
 	/** List of all client beacon actors with active connections */
 	UPROPERTY()
-	TArray<AOnlineBeaconClient*> ClientActors;
+	TArray<TObjectPtr<AOnlineBeaconClient>> ClientActors;
 };
diff --git a/Engine/Plugins/Online/OnlineSubsystemUtils/Source/OnlineSubsystemUtils/Public/PartyBeaconHost.h b/Engine/Plugins/Online/OnlineSubsystemUtils/Source/OnlineSubsystemUtils/Public/PartyBeaconHost.h
--- a/Engine/Plugins/Online/OnlineSubsystemUtils/Source/OnlineSubsystemUtils/Public/PartyBeaconHost.h
+++ b/Engine/Plugins/Online/OnlineSubsystemUtils/Source/OnlineSubsystemUtils/Public/PartyBeaconHost.h
@@ -378,7 +378,7 @@ protected:
 
 	/** State of the beacon */
 	UPROPERTY()
-	UPartyBeaconState* State;
+	TObjectPtr<UPartyBeaconState> State;
 
 	/** Delegate fired when the beacon indicates all reservations are taken */
 	FOnReservationUpdate ReservationsFull;
diff --git a/Engine/Plugins/Online/OnlineSubsystemUtils/Source/OnlineSubsystemUtils/Public/SpectatorBeaconHost.h b/Engine/Plugins/Online/OnlineSubsystemUtils/Source/OnlineSubsystemUtils/Public/SpectatorBeaconHost.h
--- a/Engine/Plugins/Online/OnlineSubsystemUtils/Source/OnlineSubsystemUtils/Public/SpectatorBeaconHost.h
+++ b/Engine/Plugins/Online/OnlineSubsystemUtils/Source/OnlineSubsystemUtils/Public/SpectatorBeaconHost.h
@@ -223,7 +223,7 @@ protected:
 
 	/** State of the beacon */
 	UPROPERTY()
-		USpectatorBeaconState* State;
+		TObjectPtr<USpectatorBeaconState> State;
 
 	/** Delegate fired when the beacon indicates all reservations are taken */
 	FOnReservationUpdate ReservationsFull;
diff --git a/Engine/Plugins/Runtime/AudioSynesthesia/Source/AudioSynesthesia/Classes/ConstantQNRT.h b/Engine/Plugins/Runtime/AudioSynesthesia/Source/AudioSynesthesia/Classes/ConstantQNRT.h
--- a/Engine/Plugins/Runtime/AudioSynesthesia/Source/AudioSynesthesia/Classes/ConstantQNRT.h
+++ b/Engine/Plugins/Runtime/AudioSynesthesia/Source/AudioSynesthesia/Classes/ConstantQNRT.h
@@ -141,7 +141,7 @@ class AUDIOSYNESTHESIA_API UConstantQNRT : public UAudioSynesthesiaNRT
 
 		/** The settings for the audio analyzer.  */
 		UPROPERTY(EditAnywhere, BlueprintReadOnly, Category=AudioAnalyzer)
-		UConstantQNRTSettings* Settings;
+		TObjectPtr<UConstantQNRTSettings> Settings;
 
 		/** Get a specific channel cqt of the analyzed sound at a given time. */
 		UFUNCTION(BlueprintCallable, Category="Audio Analyzer")
diff --git a/Engine/Plugins/Runtime/AudioSynesthesia/Source/AudioSynesthesia/Classes/Loudness.h b/Engine/Plugins/Runtime/AudioSynesthesia/Source/AudioSynesthesia/Classes/Loudness.h
--- a/Engine/Plugins/Runtime/AudioSynesthesia/Source/AudioSynesthesia/Classes/Loudness.h
+++ b/Engine/Plugins/Runtime/AudioSynesthesia/Source/AudioSynesthesia/Classes/Loudness.h
@@ -124,7 +124,7 @@ public:
 
 	/** The settings for the audio analyzer.  */
 	UPROPERTY(EditAnywhere, BlueprintReadWrite, Category = AudioAnalyzer)
-	ULoudnessSettings* Settings;
+	TObjectPtr<ULoudnessSettings> Settings;
 
 	/** Delegate to receive all overall loudness results since last delegate call. */
 	UPROPERTY(BlueprintAssignable)
diff --git a/Engine/Plugins/Runtime/AudioSynesthesia/Source/AudioSynesthesia/Classes/LoudnessNRT.h b/Engine/Plugins/Runtime/AudioSynesthesia/Source/AudioSynesthesia/Classes/LoudnessNRT.h
--- a/Engine/Plugins/Runtime/AudioSynesthesia/Source/AudioSynesthesia/Classes/LoudnessNRT.h
+++ b/Engine/Plugins/Runtime/AudioSynesthesia/Source/AudioSynesthesia/Classes/LoudnessNRT.h
@@ -85,7 +85,7 @@ class AUDIOSYNESTHESIA_API ULoudnessNRT : public UAudioSynesthesiaNRT
 
 		/** The settings for the audio analyzer.  */
 		UPROPERTY(EditAnywhere, BlueprintReadOnly, Category=AudioAnalyzer)
-		ULoudnessNRTSettings* Settings;
+		TObjectPtr<ULoudnessNRTSettings> Settings;
 
 		/** Get the overall loudness of the analyzed sound at a given time. */
 		UFUNCTION(BlueprintCallable, Category="Audio Analyzer")
diff --git a/Engine/Plugins/Runtime/AudioSynesthesia/Source/AudioSynesthesia/Classes/Meter.h b/Engine/Plugins/Runtime/AudioSynesthesia/Source/AudioSynesthesia/Classes/Meter.h
--- a/Engine/Plugins/Runtime/AudioSynesthesia/Source/AudioSynesthesia/Classes/Meter.h
+++ b/Engine/Plugins/Runtime/AudioSynesthesia/Source/AudioSynesthesia/Classes/Meter.h
@@ -127,7 +127,7 @@ public:
 
 	/** The settings for the meter audio analyzer.  */
 	UPROPERTY(EditAnywhere, BlueprintReadWrite, Category = AudioAnalyzer)
-	UMeterSettings* Settings;
+	TObjectPtr<UMeterSettings> Settings;
 
 	/** Delegate to receive all overall meter results since last delegate call. */
 	UPROPERTY(BlueprintAssignable)
diff --git a/Engine/Plugins/Runtime/AudioSynesthesia/Source/AudioSynesthesia/Classes/OnsetNRT.h b/Engine/Plugins/Runtime/AudioSynesthesia/Source/AudioSynesthesia/Classes/OnsetNRT.h
--- a/Engine/Plugins/Runtime/AudioSynesthesia/Source/AudioSynesthesia/Classes/OnsetNRT.h
+++ b/Engine/Plugins/Runtime/AudioSynesthesia/Source/AudioSynesthesia/Classes/OnsetNRT.h
@@ -64,7 +64,7 @@ class AUDIOSYNESTHESIA_API UOnsetNRT : public UAudioSynesthesiaNRT
 
 		/** The settings for the audio analyzer.  */
 		UPROPERTY(EditAnywhere, BlueprintReadOnly, Category=AudioAnalyzer)
-		UOnsetNRTSettings* Settings;
+		TObjectPtr<UOnsetNRTSettings> Settings;
 
 		/** Returns onsets which occured between start and end timestamps. */
 		UFUNCTION(BlueprintCallable, Category="Audio Analyzer")
diff --git a/Engine/Plugins/Runtime/CableComponent/Source/CableComponent/Classes/CableActor.h b/Engine/Plugins/Runtime/CableComponent/Source/CableComponent/Classes/CableActor.h
--- a/Engine/Plugins/Runtime/CableComponent/Source/CableComponent/Classes/CableActor.h
+++ b/Engine/Plugins/Runtime/CableComponent/Source/CableComponent/Classes/CableActor.h
@@ -15,5 +15,5 @@ class CABLECOMPONENT_API ACableActor : public AActor
 
 	/** Cable component that performs simulation and rendering */
 	UPROPERTY(Category=Cable, VisibleAnywhere, BlueprintReadOnly)
-	class UCableComponent* CableComponent;
+	TObjectPtr<class UCableComponent> CableComponent;
 };
diff --git a/Engine/Plugins/Runtime/Oculus/OculusVR/Source/OculusHMD/Private/OculusResourceHolder.h b/Engine/Plugins/Runtime/Oculus/OculusVR/Source/OculusHMD/Private/OculusResourceHolder.h
--- a/Engine/Plugins/Runtime/Oculus/OculusVR/Source/OculusHMD/Private/OculusResourceHolder.h
+++ b/Engine/Plugins/Runtime/Oculus/OculusVR/Source/OculusHMD/Private/OculusResourceHolder.h
@@ -17,5 +17,5 @@ class UOculusResourceHolder : public UObject
 
 public:
 	UPROPERTY()
-	UMaterial* PokeAHoleMaterial;
+	TObjectPtr<UMaterial> PokeAHoleMaterial;
 };
\ No newline at end of file
diff --git a/Engine/Plugins/Runtime/Oculus/OculusVR/Source/OculusHMD/Private/OculusSceneCaptureCubemap.h b/Engine/Plugins/Runtime/Oculus/OculusVR/Source/OculusHMD/Private/OculusSceneCaptureCubemap.h
--- a/Engine/Plugins/Runtime/Oculus/OculusVR/Source/OculusHMD/Private/OculusSceneCaptureCubemap.h
+++ b/Engine/Plugins/Runtime/Oculus/OculusVR/Source/OculusHMD/Private/OculusSceneCaptureCubemap.h
@@ -67,7 +67,7 @@ private:
 	} Stage;
 
 	UPROPERTY()
-	TArray<USceneCaptureComponent2D*> CaptureComponents;
+	TArray<TObjectPtr<USceneCaptureComponent2D>> CaptureComponents;
 
 	uint32 CaptureBoxSideRes;
 	EPixelFormat CaptureFormat;
diff --git a/Engine/Plugins/Runtime/Oculus/OculusVR/Source/OculusInput/Public/OculusHandComponent.h b/Engine/Plugins/Runtime/Oculus/OculusVR/Source/OculusInput/Public/OculusHandComponent.h
--- a/Engine/Plugins/Runtime/Oculus/OculusVR/Source/OculusInput/Public/OculusHandComponent.h
+++ b/Engine/Plugins/Runtime/Oculus/OculusVR/Source/OculusInput/Public/OculusHandComponent.h
@@ -49,7 +49,7 @@ public:
 
 	/** Material that gets applied to the hands when the system gesture is active */
 	UPROPERTY(EditDefaultsOnly, BlueprintReadWrite, Category = "HandProperties")
-	class UMaterialInterface* SystemGestureMaterial;
+	TObjectPtr<class UMaterialInterface> SystemGestureMaterial;
 
 	/** Whether or not to initialize physics capsules on the skeletal mesh */
 	UPROPERTY(EditDefaultsOnly, BlueprintReadWrite, Category = "HandProperties")
@@ -61,7 +61,7 @@ public:
 
 	/** Material override for the runtime skeletal mesh */
 	UPROPERTY(EditDefaultsOnly, BlueprintReadWrite, Category = "HandProperties")
-	class UMaterialInterface* MaterialOverride;
+	TObjectPtr<class UMaterialInterface> MaterialOverride;
 
 	/** Bone mapping for custom hand skeletal meshes */
 	UPROPERTY(EditAnywhere, BlueprintReadWrite, Category = "CustomSkeletalMesh")
diff --git a/Engine/Plugins/Runtime/Oculus/OculusVR/Source/OculusInput/Public/OculusInputFunctionLibrary.h b/Engine/Plugins/Runtime/Oculus/OculusVR/Source/OculusInput/Public/OculusInputFunctionLibrary.h
--- a/Engine/Plugins/Runtime/Oculus/OculusVR/Source/OculusInput/Public/OculusInputFunctionLibrary.h
+++ b/Engine/Plugins/Runtime/Oculus/OculusVR/Source/OculusInput/Public/OculusInputFunctionLibrary.h
@@ -74,7 +74,7 @@ struct OCULUSINPUT_API FOculusCapsuleCollider
 
 public:
 	UPROPERTY(BlueprintReadOnly, Category = "OculusLibrary|HandTracking")
-	UCapsuleComponent* Capsule { nullptr };
+	TObjectPtr<UCapsuleComponent> Capsule { nullptr };
 
 	UPROPERTY(BlueprintReadOnly, Category = "OculusLibrary|HandTracking")
 	EBone BoneId = EBone::Wrist_Root;
diff --git a/Engine/Plugins/Runtime/Oculus/OculusVR/Source/OculusMR/Private/OculusMR_CastingCameraActor.h b/Engine/Plugins/Runtime/Oculus/OculusVR/Source/OculusMR/Private/OculusMR_CastingCameraActor.h
--- a/Engine/Plugins/Runtime/Oculus/OculusVR/Source/OculusMR/Private/OculusMR_CastingCameraActor.h
+++ b/Engine/Plugins/Runtime/Oculus/OculusVR/Source/OculusMR/Private/OculusMR_CastingCameraActor.h
@@ -41,34 +41,34 @@ public:
 	virtual void BeginDestroy() override;
 
 	UPROPERTY()
-	class UVRNotificationsComponent* VRNotificationComponent;
+	TObjectPtr<class UVRNotificationsComponent> VRNotificationComponent;
 
 	UPROPERTY()
-	UTexture2D* CameraColorTexture;
+	TObjectPtr<UTexture2D> CameraColorTexture;
 
 	UPROPERTY()
-	UTexture2D* CameraDepthTexture;
+	TObjectPtr<UTexture2D> CameraDepthTexture;
 
 	UPROPERTY()
-	UOculusMR_PlaneMeshComponent* PlaneMeshComponent;
+	TObjectPtr<UOculusMR_PlaneMeshComponent> PlaneMeshComponent;
 
 	UPROPERTY()
-	UMaterial* ChromaKeyMaterial;
+	TObjectPtr<UMaterial> ChromaKeyMaterial;
 
 	UPROPERTY()
-	UMaterial* OpaqueColoredMaterial;
+	TObjectPtr<UMaterial> OpaqueColoredMaterial;
 
 	UPROPERTY()
-	UMaterialInstanceDynamic* ChromaKeyMaterialInstance;
+	TObjectPtr<UMaterialInstanceDynamic> ChromaKeyMaterialInstance;
 
 	UPROPERTY()
-	UMaterialInstanceDynamic* CameraFrameMaterialInstance;
+	TObjectPtr<UMaterialInstanceDynamic> CameraFrameMaterialInstance;
 
 	UPROPERTY()
-	UMaterialInstanceDynamic* BackdropMaterialInstance;
+	TObjectPtr<UMaterialInstanceDynamic> BackdropMaterialInstance;
 
 	UPROPERTY()
-	class UTexture2D* DefaultTexture_White;
+	TObjectPtr<class UTexture2D> DefaultTexture_White;
 
 	bool TrackedCameraCalibrationRequired;
 	bool HasTrackedCameraCalibrationCalibrated;
@@ -117,22 +117,22 @@ private:
 	float ForegroundMaxDistance;
 
 	UPROPERTY()
-	TArray<UTextureRenderTarget2D*> BackgroundRenderTargets;
+	TArray<TObjectPtr<UTextureRenderTarget2D>> BackgroundRenderTargets;
 
 	UPROPERTY()
-	ASceneCapture2D* ForegroundCaptureActor;
+	TObjectPtr<ASceneCapture2D> ForegroundCaptureActor;
 
 	UPROPERTY()
-	TArray<UTextureRenderTarget2D*> ForegroundRenderTargets;
+	TArray<TObjectPtr<UTextureRenderTarget2D>> ForegroundRenderTargets;
 
 	UPROPERTY()
 	TArray<double> PoseTimes;
 
 	UPROPERTY()
-	UOculusMR_Settings* MRSettings;
+	TObjectPtr<UOculusMR_Settings> MRSettings;
 
 	UPROPERTY()
-	UOculusMR_State* MRState;
+	TObjectPtr<UOculusMR_State> MRState;
 
 #if PLATFORM_ANDROID
 	TArray<Audio::AlignedFloatBuffer> AudioBuffers;
diff --git a/Engine/Plugins/Runtime/Oculus/OculusVR/Source/OculusMR/Private/OculusMR_State.h b/Engine/Plugins/Runtime/Oculus/OculusVR/Source/OculusMR/Private/OculusMR_State.h
--- a/Engine/Plugins/Runtime/Oculus/OculusVR/Source/OculusMR/Private/OculusMR_State.h
+++ b/Engine/Plugins/Runtime/Oculus/OculusVR/Source/OculusMR/Private/OculusMR_State.h
@@ -100,7 +100,7 @@ public:
 
 	// Component at the tracking origin that the camera calibration is applied to
 	UPROPERTY()
-	class USceneComponent* TrackingReferenceComponent;
+	TObjectPtr<class USceneComponent> TrackingReferenceComponent;
 
 	// A multiplier on the camera distance, should be based on the scaling of the player component
 	UPROPERTY()
diff --git a/Engine/Plugins/Runtime/ProceduralMeshComponent/Source/ProceduralMeshComponent/Public/ProceduralMeshComponent.h b/Engine/Plugins/Runtime/ProceduralMeshComponent/Source/ProceduralMeshComponent/Public/ProceduralMeshComponent.h
--- a/Engine/Plugins/Runtime/ProceduralMeshComponent/Source/ProceduralMeshComponent/Public/ProceduralMeshComponent.h
+++ b/Engine/Plugins/Runtime/ProceduralMeshComponent/Source/ProceduralMeshComponent/Public/ProceduralMeshComponent.h
@@ -280,7 +280,7 @@ public:
 
 	/** Collision data */
 	UPROPERTY(Instanced)
-	class UBodySetup* ProcMeshBodySetup;
+	TObjectPtr<class UBodySetup> ProcMeshBodySetup;
 
 	/** 
 	 *	Get pointer to internal data for one section of this procedural mesh component. 
@@ -338,7 +338,7 @@ private:
 	
 	/** Queue for async body setups that are being cooked */
 	UPROPERTY(transient)
-	TArray<UBodySetup*> AsyncBodySetupQueue;
+	TArray<TObjectPtr<UBodySetup>> AsyncBodySetupQueue;
 
 	friend class FProceduralMeshSceneProxy;
 };
diff --git a/Engine/Plugins/Runtime/ReplicationGraph/Source/Public/BasicReplicationGraph.h b/Engine/Plugins/Runtime/ReplicationGraph/Source/Public/BasicReplicationGraph.h
--- a/Engine/Plugins/Runtime/ReplicationGraph/Source/Public/BasicReplicationGraph.h
+++ b/Engine/Plugins/Runtime/ReplicationGraph/Source/Public/BasicReplicationGraph.h
@@ -15,10 +15,10 @@ struct FConnectionAlwaysRelevantNodePair
 	bool operator==(UNetConnection* InConnection) const;
 
 	UPROPERTY()
-	UNetConnection* NetConnection = nullptr;
+	TObjectPtr<UNetConnection> NetConnection = nullptr;
 
 	UPROPERTY()
-	UReplicationGraphNode_AlwaysRelevant_ForConnection* Node = nullptr;	
+	TObjectPtr<UReplicationGraphNode_AlwaysRelevant_ForConnection> Node = nullptr;	
 };
 
 
@@ -49,17 +49,17 @@ public:
 	virtual int32 ServerReplicateActors(float DeltaSeconds) override;
 
 	UPROPERTY()
-	UReplicationGraphNode_GridSpatialization2D* GridNode;
+	TObjectPtr<UReplicationGraphNode_GridSpatialization2D> GridNode;
 
 	UPROPERTY()
-	UReplicationGraphNode_ActorList* AlwaysRelevantNode;
+	TObjectPtr<UReplicationGraphNode_ActorList> AlwaysRelevantNode;
 
 	UPROPERTY()
 	TArray<FConnectionAlwaysRelevantNodePair> AlwaysRelevantForConnectionList;
 
 	/** Actors that are only supposed to replicate to their owning connection, but that did not have a connection on spawn */
 	UPROPERTY()
-	TArray<AActor*> ActorsWithoutNetConnection;
+	TArray<TObjectPtr<AActor>> ActorsWithoutNetConnection;
 
 
 	UReplicationGraphNode_AlwaysRelevant_ForConnection* GetAlwaysRelevantNodeForConnection(UNetConnection* Connection);
diff --git a/Engine/Plugins/Runtime/ReplicationGraph/Source/Public/ReplicationGraph.h b/Engine/Plugins/Runtime/ReplicationGraph/Source/Public/ReplicationGraph.h
--- a/Engine/Plugins/Runtime/ReplicationGraph/Source/Public/ReplicationGraph.h
+++ b/Engine/Plugins/Runtime/ReplicationGraph/Source/Public/ReplicationGraph.h
@@ -145,7 +145,7 @@ protected:
 protected:
 
 	UPROPERTY()
-	TArray< UReplicationGraphNode* > AllChildNodes;
+	TArray< TObjectPtr<UReplicationGraphNode> > AllChildNodes;
 
 	TSharedPtr<FReplicationGraphGlobalData>	GraphGlobals;
 
@@ -543,10 +543,10 @@ public:
 private:
 
 	UPROPERTY()
-	UReplicationGraphNode* DynamicNode = nullptr;
+	TObjectPtr<UReplicationGraphNode> DynamicNode = nullptr;
 
 	UPROPERTY()
-	UReplicationGraphNode_DormancyNode* DormancyNode = nullptr;
+	TObjectPtr<UReplicationGraphNode_DormancyNode> DormancyNode = nullptr;
 
 	UReplicationGraphNode* GetDynamicNode();
 
@@ -754,7 +754,7 @@ public:
 protected:
 
 	UPROPERTY()
-	UReplicationGraphNode* ChildNode = nullptr;
+	TObjectPtr<UReplicationGraphNode> ChildNode = nullptr;
 
 	// Explicit list of classes that are always relevant. This 
 	TArray<UClass*>	AlwaysRelevantClasses;
@@ -768,13 +768,13 @@ struct FAlwaysRelevantActorInfo
 	GENERATED_BODY()
 
 	UPROPERTY()
-	const UNetConnection* Connection = nullptr;
+	TObjectPtr<const UNetConnection> Connection = nullptr;
 
 	UPROPERTY()
-	AActor* LastViewer = nullptr;
+	TObjectPtr<AActor> LastViewer = nullptr;
 
 	UPROPERTY()
-	AActor* LastViewTarget = nullptr;
+	TObjectPtr<AActor> LastViewTarget = nullptr;
 
 	bool operator==(UNetConnection* Other) const
 	{
@@ -820,7 +820,7 @@ struct FTearOffActorInfo
 	uint32 TearOffFrameNum;
 
 	UPROPERTY()
-	AActor* Actor;
+	TObjectPtr<AActor> Actor;
 
 	bool bHasReppedOnce;
 };
@@ -869,15 +869,15 @@ public:
 	TSubclassOf<UNetReplicationGraphConnection> ReplicationConnectionManagerClass;
 
 	UPROPERTY()
-	UNetDriver*	NetDriver;
+	TObjectPtr<UNetDriver>	NetDriver;
 
 	/** List of connection managers. This list is not sorted and not stable. */
 	UPROPERTY()
-	TArray<UNetReplicationGraphConnection*> Connections;
+	TArray<TObjectPtr<UNetReplicationGraphConnection>> Connections;
 
 	/** ConnectionManagers that we have created but haven't officially been added to the net driver's ClientConnection list. This is a  hack for silly order of initialization issues. */
 	UPROPERTY()
-	TArray<UNetReplicationGraphConnection*> PendingConnections;
+	TArray<TObjectPtr<UNetReplicationGraphConnection>> PendingConnections;
 
 	/** The max distance between an FActorDestructionInfo and a connection that we will replicate. */
 	float DestructInfoMaxDistanceSquared = 15000.f * 15000.f;
@@ -1035,11 +1035,11 @@ protected:
 	
 	/** A list of nodes that can add actors to all connections. They don't necessarily *have to* add actors to each connection, but they will get a chance to. These are added via ::AddGlobalGraphNode  */
 	UPROPERTY()
-	TArray<UReplicationGraphNode*> GlobalGraphNodes;
+	TArray<TObjectPtr<UReplicationGraphNode>> GlobalGraphNodes;
 
 	/** A list of nodes that want PrepareForReplication() to be called on them at the top of the replication frame. */
 	UPROPERTY()
-	TArray<UReplicationGraphNode*> PrepareForReplicationNodes;
+	TArray<TObjectPtr<UReplicationGraphNode>> PrepareForReplicationNodes;
 	
 	FGlobalActorReplicationInfoMap GlobalActorReplicationInfoMap;
 
@@ -1107,7 +1107,7 @@ struct FLastLocationGatherInfo
 	FLastLocationGatherInfo(const UNetConnection* InConnection, FVector InLastLocation) : Connection(InConnection), LastLocation(InLastLocation), LastOutOfRangeLocationCheck(InLastLocation)  {}
 
 	UPROPERTY()
-	const UNetConnection* Connection;
+	TObjectPtr<const UNetConnection> Connection;
 
 	UPROPERTY()
 	FVector LastLocation;
@@ -1132,7 +1132,7 @@ public:
 	UNetReplicationGraphConnection();
 
 	UPROPERTY()
-	UNetConnection* NetConnection;
+	TObjectPtr<UNetConnection> NetConnection;
 
 	/** A map of all of our per-actor data */
 	FPerConnectionActorInfoMap ActorInfoMap;
@@ -1153,7 +1153,7 @@ public:
 #endif
 
 	UPROPERTY()
-	class AReplicationGraphDebugActor* DebugActor = nullptr;
+	TObjectPtr<class AReplicationGraphDebugActor> DebugActor = nullptr;
 
 	bool bEnableDebugging;
 
@@ -1235,10 +1235,10 @@ private:
 	void OnUpdateViewerLocation(FLastLocationGatherInfo* LocationInfo, const FNetViewer& Viewer, const FReplicationGraphDestructionSettings& DestructionSettings);
 
 	UPROPERTY()
-	TArray<UReplicationGraphNode*> ConnectionGraphNodes;
+	TArray<TObjectPtr<UReplicationGraphNode>> ConnectionGraphNodes;
 
 	UPROPERTY()
-	UReplicationGraphNode_TearOff_ForConnection* TearOffNode;
+	TObjectPtr<UReplicationGraphNode_TearOff_ForConnection> TearOffNode;
 	
 	/** DestructionInfo handling. This is how we send "delete this actor" to clients when the actor is deleted on the server (placed in map actors) */
 	struct FCachedDestructInfo
@@ -1307,7 +1307,7 @@ public:
 	virtual bool IsReplayRelevantFor(const AActor* RealViewer, const AActor* ViewTarget, const FVector& SrcLocation, const float CullDistanceSquared) const override { return false; }
 
 	UPROPERTY()
-	UReplicationGraph* ReplicationGraph;
+	TObjectPtr<UReplicationGraph> ReplicationGraph;
 
 	UFUNCTION(Server, Reliable)
 	void ServerStartDebugging();
@@ -1342,7 +1342,7 @@ public:
 
 
 	UPROPERTY()
-	UNetReplicationGraphConnection* ConnectionManager;
+	TObjectPtr<UNetReplicationGraphConnection> ConnectionManager;
 
 	virtual UNetConnection* GetNetConnection() const override;
 };
diff --git a/Engine/Plugins/Runtime/Steam/SteamVR/Source/SteamVRInputDevice/Public/SteamVRTrackingRefComponent.h b/Engine/Plugins/Runtime/Steam/SteamVR/Source/SteamVRInputDevice/Public/SteamVRTrackingRefComponent.h
--- a/Engine/Plugins/Runtime/Steam/SteamVR/Source/SteamVRInputDevice/Public/SteamVRTrackingRefComponent.h
+++ b/Engine/Plugins/Runtime/Steam/SteamVR/Source/SteamVRInputDevice/Public/SteamVRTrackingRefComponent.h
@@ -75,7 +75,7 @@ public:
 
 	/** Currently displayed Tracking References in-world */
 	UPROPERTY(BlueprintReadOnly, Category = "SteamVR Input")
-	TArray<UStaticMeshComponent*> TrackingReferences;
+	TArray<TObjectPtr<UStaticMeshComponent>> TrackingReferences;
 
 private:
 	/** Represents a tracked device with a flag on whether its activated or not */
diff --git a/Engine/Plugins/Runtime/Synthesis/Source/Synthesis/Classes/SourceEffects/SourceEffectFilter.h b/Engine/Plugins/Runtime/Synthesis/Source/Synthesis/Classes/SourceEffects/SourceEffectFilter.h
--- a/Engine/Plugins/Runtime/Synthesis/Source/Synthesis/Classes/SourceEffects/SourceEffectFilter.h
+++ b/Engine/Plugins/Runtime/Synthesis/Source/Synthesis/Classes/SourceEffects/SourceEffectFilter.h
@@ -44,7 +44,7 @@ struct SYNTHESIS_API FSourceEffectFilterAudioBusModulationSettings
 
 	// Audio bus to use to modulate the filter
 	UPROPERTY(EditAnywhere, BlueprintReadWrite, Category = "SourceEffect|Preset|AudioBus")
-	UAudioBus* AudioBus = nullptr;
+	TObjectPtr<UAudioBus> AudioBus = nullptr;
 
 	// The amplitude envelope follower attack time (in milliseconds) on the audio bus.
 	UPROPERTY(EditAnywhere, BlueprintReadWrite, Category = "SourceEffect|Modulation Envelope", meta = (ClampMin = "0", UIMin = "0", UIMax = "2000.0", DisplayName = "Attack Time (ms)"))
diff --git a/Engine/Plugins/Runtime/Synthesis/Source/Synthesis/Classes/SourceEffects/SourceEffectRingModulation.h b/Engine/Plugins/Runtime/Synthesis/Source/Synthesis/Classes/SourceEffects/SourceEffectRingModulation.h
--- a/Engine/Plugins/Runtime/Synthesis/Source/Synthesis/Classes/SourceEffects/SourceEffectRingModulation.h
+++ b/Engine/Plugins/Runtime/Synthesis/Source/Synthesis/Classes/SourceEffects/SourceEffectRingModulation.h
@@ -45,7 +45,7 @@ struct SYNTHESIS_API FSourceEffectRingModulationSettings
 
 	// Audio bus to use to modulate the effect
 	UPROPERTY(EditAnywhere, BlueprintReadWrite, Category = "SourceEffect|Preset")
-	UAudioBus* AudioBusModulator;
+	TObjectPtr<UAudioBus> AudioBusModulator;
 
 	FSourceEffectRingModulationSettings()
 		: ModulatorType(ERingModulatorTypeSourceEffect::Sine)
diff --git a/Engine/Plugins/Runtime/Synthesis/Source/Synthesis/Classes/SubmixEffects/SubmixEffectConvolutionReverb.h b/Engine/Plugins/Runtime/Synthesis/Source/Synthesis/Classes/SubmixEffects/SubmixEffectConvolutionReverb.h
--- a/Engine/Plugins/Runtime/Synthesis/Source/Synthesis/Classes/SubmixEffects/SubmixEffectConvolutionReverb.h
+++ b/Engine/Plugins/Runtime/Synthesis/Source/Synthesis/Classes/SubmixEffects/SubmixEffectConvolutionReverb.h
@@ -118,7 +118,7 @@ struct SYNTHESIS_API FSubmixEffectConvolutionReverbSettings
 	float SurroundRearChannelBleedAmount_DEPRECATED;
 
 	UPROPERTY(meta = ( DeprecatedProperty ) )
-	UAudioImpulseResponse* ImpulseResponse_DEPRECATED;
+	TObjectPtr<UAudioImpulseResponse> ImpulseResponse_DEPRECATED;
 
 	UPROPERTY(meta = ( DeprecatedProperty ) )
 	bool AllowHardwareAcceleration_DEPRECATED;
@@ -249,7 +249,7 @@ public:
 
 	/** The impulse response used for convolution. */
 	UPROPERTY(EditAnywhere, BlueprintReadWrite, BlueprintSetter = SetImpulseResponse, Category = SubmixEffectPreset)
-	UAudioImpulseResponse* ImpulseResponse;
+	TObjectPtr<UAudioImpulseResponse> ImpulseResponse;
 
 	/** ConvolutionReverbPreset Preset Settings. */
 	UPROPERTY(EditAnywhere, BlueprintReadWrite, BlueprintSetter = SetSettings, Category = SubmixEffectPreset)
diff --git a/Engine/Plugins/Runtime/Synthesis/Source/Synthesis/Classes/SynthComponents/SynthComponentGranulator.h b/Engine/Plugins/Runtime/Synthesis/Source/Synthesis/Classes/SynthComponents/SynthComponentGranulator.h
--- a/Engine/Plugins/Runtime/Synthesis/Source/Synthesis/Classes/SynthComponents/SynthComponentGranulator.h
+++ b/Engine/Plugins/Runtime/Synthesis/Source/Synthesis/Classes/SynthComponents/SynthComponentGranulator.h
@@ -125,7 +125,7 @@ public:
 protected:
 
 	UPROPERTY(Transient)
-	USoundWave* GranulatedSoundWave;
+	TObjectPtr<USoundWave> GranulatedSoundWave;
 
 	Audio::FGranularSynth GranularSynth;
 	Audio::FSoundWavePCMLoader SoundWaveLoader;
diff --git a/Engine/Plugins/Runtime/Synthesis/Source/Synthesis/Classes/SynthComponents/SynthComponentMonoWaveTable.h b/Engine/Plugins/Runtime/Synthesis/Source/Synthesis/Classes/SynthComponents/SynthComponentMonoWaveTable.h
--- a/Engine/Plugins/Runtime/Synthesis/Source/Synthesis/Classes/SynthComponents/SynthComponentMonoWaveTable.h
+++ b/Engine/Plugins/Runtime/Synthesis/Source/Synthesis/Classes/SynthComponents/SynthComponentMonoWaveTable.h
@@ -395,7 +395,7 @@ protected:
 
 	/** The settings asset to use for this synth */
 	UPROPERTY(EditAnywhere, Category = "Synth|Components|Audio")
-	UMonoWaveTableSynthPreset* CurrentPreset;
+	TObjectPtr<UMonoWaveTableSynthPreset> CurrentPreset;
 
 	UMonoWaveTableSynthPreset* CachedPreset;
 
diff --git a/Engine/Plugins/Runtime/Synthesis/Source/Synthesis/Classes/SynthComponents/SynthComponentWaveTable.h b/Engine/Plugins/Runtime/Synthesis/Source/Synthesis/Classes/SynthComponents/SynthComponentWaveTable.h
--- a/Engine/Plugins/Runtime/Synthesis/Source/Synthesis/Classes/SynthComponents/SynthComponentWaveTable.h
+++ b/Engine/Plugins/Runtime/Synthesis/Source/Synthesis/Classes/SynthComponents/SynthComponentWaveTable.h
@@ -50,7 +50,7 @@ class SYNTHESIS_API USynthSamplePlayer : public USynthComponent
 public:
 
 	UPROPERTY(EditAnywhere, BlueprintReadWrite, Category = "Synth|Preset")
-	USoundWave* SoundWave;
+	TObjectPtr<USoundWave> SoundWave;
 
 	UPROPERTY(BlueprintAssignable, Category = "Synth|Components|Audio")
 	FOnSampleLoaded OnSampleLoaded;
diff --git a/Engine/Source/Developer/AITestSuite/Classes/MockAI.h b/Engine/Source/Developer/AITestSuite/Classes/MockAI.h
--- a/Engine/Source/Developer/AITestSuite/Classes/MockAI.h
+++ b/Engine/Source/Developer/AITestSuite/Classes/MockAI.h
@@ -38,16 +38,16 @@ class UMockAI : public UObject
 	FTestTickHelper TickHelper;
 
 	UPROPERTY()
-	UBlackboardComponent* BBComp;
+	TObjectPtr<UBlackboardComponent> BBComp;
 
 	UPROPERTY()
-	UBrainComponent* BrainComp;
+	TObjectPtr<UBrainComponent> BrainComp;
 
 	UPROPERTY()
-	UAIPerceptionComponent* PerceptionComp;
+	TObjectPtr<UAIPerceptionComponent> PerceptionComp;
 
 	UPROPERTY()
-	UPawnActionsComponent* PawnActionComp;
+	TObjectPtr<UPawnActionsComponent> PawnActionComp;
 	
 	template<typename TBrainClass>
 	void UseBrainComponent()
diff --git a/Engine/Source/Developer/AITestSuite/Classes/MockAI_BT.h b/Engine/Source/Developer/AITestSuite/Classes/MockAI_BT.h
--- a/Engine/Source/Developer/AITestSuite/Classes/MockAI_BT.h
+++ b/Engine/Source/Developer/AITestSuite/Classes/MockAI_BT.h
@@ -16,7 +16,7 @@ class UMockAI_BT : public UMockAI
 	GENERATED_UCLASS_BODY()
 
 	UPROPERTY()
-	UBehaviorTreeComponent* BTComp;
+	TObjectPtr<UBehaviorTreeComponent> BTComp;
 
 	static TArray<int32> ExecutionLog;
 	TArray<int32> ExpectedResult;
diff --git a/Engine/Source/Developer/AITestSuite/Classes/MockGameplayTasks.h b/Engine/Source/Developer/AITestSuite/Classes/MockGameplayTasks.h
--- a/Engine/Source/Developer/AITestSuite/Classes/MockGameplayTasks.h
+++ b/Engine/Source/Developer/AITestSuite/Classes/MockGameplayTasks.h
@@ -73,7 +73,7 @@ class UMockGameplayTaskOwner : public UObject, public IGameplayTaskOwnerInterfac
 
 public:
 	UPROPERTY()
-	UGameplayTasksComponent* GTComponent;
+	TObjectPtr<UGameplayTasksComponent> GTComponent;
 
 	virtual UGameplayTasksComponent* GetGameplayTasksComponent(const UGameplayTask& Task) const override { return GTComponent; }	
 	virtual AActor* GetGameplayTaskOwner(const UGameplayTask* Task) const { return nullptr; }
diff --git a/Engine/Source/Developer/BlueprintNativeCodeGen/Private/BlueprintNativeCodeGenManifest.h b/Engine/Source/Developer/BlueprintNativeCodeGen/Private/BlueprintNativeCodeGenManifest.h
--- a/Engine/Source/Developer/BlueprintNativeCodeGen/Private/BlueprintNativeCodeGenManifest.h
+++ b/Engine/Source/Developer/BlueprintNativeCodeGen/Private/BlueprintNativeCodeGenManifest.h
@@ -25,7 +25,7 @@ public:
 
 public:
 	UPROPERTY()
-	UClass* AssetType;
+	TObjectPtr<UClass> AssetType;
 
 	// cannot use a FSoftObjectPath, as the json serializer has problems 
 	// with some asset paths (for example, I had a folder named 'Folder()')
@@ -230,7 +230,7 @@ private:
 	FString OutputDir;
 
 	UPROPERTY()
-	TArray<UPackage*> ModuleDependencies;
+	TArray<TObjectPtr<UPackage>> ModuleDependencies;
 
 	/**  */
 	UPROPERTY()
diff --git a/Engine/Source/Developer/FunctionalTesting/Classes/FunctionalAITest.h b/Engine/Source/Developer/FunctionalTesting/Classes/FunctionalAITest.h
--- a/Engine/Source/Developer/FunctionalTesting/Classes/FunctionalAITest.h
+++ b/Engine/Source/Developer/FunctionalTesting/Classes/FunctionalAITest.h
@@ -34,11 +34,11 @@ struct FAITestSpawnInfo
 
 	/** if set will be applied to spawned AI */
 	UPROPERTY(EditAnywhere, BlueprintReadWrite, Category=AISpawn)
-	class UBehaviorTree* BehaviorTree;
+	TObjectPtr<class UBehaviorTree> BehaviorTree;
 
 	/** Where should AI be spawned */
 	UPROPERTY(EditAnywhere, BlueprintReadWrite, Category=AISpawn)
-	AActor* SpawnLocation;
+	TObjectPtr<AActor> SpawnLocation;
 
 	UPROPERTY(EditAnywhere, BlueprintReadWrite, Category=AISpawn, meta=(UIMin=1, ClampMin=1))
 	int32 NumberToSpawn;
@@ -102,7 +102,7 @@ struct FAITestSpawnSet
 
 	/** location used for spawning if spawn info doesn't define one */
 	UPROPERTY(EditAnywhere, BlueprintReadWrite, Category = AISpawn)
-	AActor* FallbackSpawnLocation;
+	TObjectPtr<AActor> FallbackSpawnLocation;
 
 	FAITestSpawnSet()
 		: bEnabled(true)
@@ -123,7 +123,7 @@ protected:
 	float SpawnLocationRandomizationRange;
 
 	UPROPERTY(BlueprintReadOnly, Category=AITest)
-	TArray<APawn*> SpawnedPawns;
+	TArray<TObjectPtr<APawn>> SpawnedPawns;
 
 	UPROPERTY(BlueprintReadOnly, Category = AITest)
 	TArray<FPendingDelayedSpawn> PendingDelayedSpawns;
diff --git a/Engine/Source/Developer/FunctionalTesting/Classes/FunctionalTest.h b/Engine/Source/Developer/FunctionalTesting/Classes/FunctionalTest.h
--- a/Engine/Source/Developer/FunctionalTesting/Classes/FunctionalTest.h
+++ b/Engine/Source/Developer/FunctionalTesting/Classes/FunctionalTest.h
@@ -226,7 +226,7 @@ public:
 
 private:
 	UPROPERTY()
-	UBillboardComponent* SpriteComponent;
+	TObjectPtr<UBillboardComponent> SpriteComponent;
 
 protected:
 	/**
@@ -268,7 +268,7 @@ protected:
 	 * way to observe the test from a different location than you place the functional test actor.
 	 */
 	UPROPERTY(EditAnywhere, BlueprintReadWrite, Category="Functional Testing")
-	AActor* ObservationPoint;
+	TObjectPtr<AActor> ObservationPoint;
 
 	/**
 	 * A random number stream that you can use during testing.  This number stream will be consistent
@@ -317,16 +317,16 @@ public:
 	FFunctionalTestEventSignature OnTestFinished;
 
 	UPROPERTY(Transient)
-	TArray<AActor*> AutoDestroyActors;
+	TArray<TObjectPtr<AActor>> AutoDestroyActors;
 	
 	FString FailureMessage;
 	
 #if WITH_EDITORONLY_DATA
 	UPROPERTY()
-	class UFuncTestRenderingComponent* RenderComp;
+	TObjectPtr<class UFuncTestRenderingComponent> RenderComp;
 
 	UPROPERTY()
-	class UTextRenderComponent* TestName;
+	TObjectPtr<class UTextRenderComponent> TestName;
 #endif // WITH_EDITORONLY_DATA
 
 	/** List of causes we need a re-run. */
diff --git a/Engine/Source/Developer/FunctionalTesting/Classes/FunctionalTestingManager.h b/Engine/Source/Developer/FunctionalTesting/Classes/FunctionalTestingManager.h
--- a/Engine/Source/Developer/FunctionalTesting/Classes/FunctionalTestingManager.h
+++ b/Engine/Source/Developer/FunctionalTesting/Classes/FunctionalTestingManager.h
@@ -24,10 +24,10 @@ class UFunctionalTestingManager : public UBlueprintFunctionLibrary
 	GENERATED_UCLASS_BODY()
 	
 	UPROPERTY(Transient)
-	TArray<class AFunctionalTest*> TestsLeft;
+	TArray<TObjectPtr<class AFunctionalTest>> TestsLeft;
 
 	UPROPERTY(Transient)
-	TArray<class AFunctionalTest*> AllTests;
+	TArray<TObjectPtr<class AFunctionalTest>> AllTests;
 
 	UPROPERTY(BlueprintAssignable)
 	FFunctionalTestEventSignature OnSetupTests;
diff --git a/Engine/Source/Developer/FunctionalTesting/Classes/FunctionalUIScreenshotTest.h b/Engine/Source/Developer/FunctionalTesting/Classes/FunctionalUIScreenshotTest.h
--- a/Engine/Source/Developer/FunctionalTesting/Classes/FunctionalUIScreenshotTest.h
+++ b/Engine/Source/Developer/FunctionalTesting/Classes/FunctionalUIScreenshotTest.h
@@ -43,11 +43,11 @@ protected:
 	TSubclassOf<UUserWidget> WidgetClass;
 
 	UPROPERTY()
-	UUserWidget* SpawnedWidget;
+	TObjectPtr<UUserWidget> SpawnedWidget;
 
 	UPROPERTY(EditAnywhere, Category = "UI")
 	EWidgetTestAppearLocation WidgetLocation;
 
 	UPROPERTY(Transient, DuplicateTransient)
-	UTextureRenderTarget2D* ScreenshotRT;
+	TObjectPtr<UTextureRenderTarget2D> ScreenshotRT;
 };
diff --git a/Engine/Source/Developer/FunctionalTesting/Classes/ScreenshotFunctionalTestBase.h b/Engine/Source/Developer/FunctionalTesting/Classes/ScreenshotFunctionalTestBase.h
--- a/Engine/Source/Developer/FunctionalTesting/Classes/ScreenshotFunctionalTestBase.h
+++ b/Engine/Source/Developer/FunctionalTesting/Classes/ScreenshotFunctionalTestBase.h
@@ -66,7 +66,7 @@ protected:
 	FString Notes;
 
 	UPROPERTY(EditAnywhere, BlueprintReadOnly, Category = "Screenshot")
-	class UCameraComponent* ScreenshotCamera;
+	TObjectPtr<class UCameraComponent> ScreenshotCamera;
 
 	UPROPERTY(EditAnywhere, BlueprintReadOnly, Category = "Screenshot", SimpleDisplay)
 	FAutomationScreenshotOptions ScreenshotOptions;
diff --git a/Engine/Source/Developer/FunctionalTesting/Public/AutomationScreenshotOptions.h b/Engine/Source/Developer/FunctionalTesting/Public/AutomationScreenshotOptions.h
--- a/Engine/Source/Developer/FunctionalTesting/Public/AutomationScreenshotOptions.h
+++ b/Engine/Source/Developer/FunctionalTesting/Public/AutomationScreenshotOptions.h
@@ -152,7 +152,7 @@ public:
 	 * screenshot.
 	 */
 	UPROPERTY(EditAnywhere, BlueprintReadWrite, Category = "Screenshot")
-	UAutomationViewSettings* ViewSettings;
+	TObjectPtr<UAutomationViewSettings> ViewSettings;
 
 	/**
 	 * Allows you to screenshot a buffer other than the default final lit scene image.  Useful if you're
diff --git a/Engine/Source/Developer/FunctionalTesting/Public/GroundTruthData.h b/Engine/Source/Developer/FunctionalTesting/Public/GroundTruthData.h
--- a/Engine/Source/Developer/FunctionalTesting/Public/GroundTruthData.h
+++ b/Engine/Source/Developer/FunctionalTesting/Public/GroundTruthData.h
@@ -40,5 +40,5 @@ public:
 protected:
 	
 	UPROPERTY(VisibleAnywhere, Instanced, Category=Data)
-	UObject* ObjectData;
+	TObjectPtr<UObject> ObjectData;
 };
diff --git a/Engine/Source/Developer/GameplayDebugger/Public/GameplayDebuggerCategoryReplicator.h b/Engine/Source/Developer/GameplayDebugger/Public/GameplayDebuggerCategoryReplicator.h
--- a/Engine/Source/Developer/GameplayDebugger/Public/GameplayDebuggerCategoryReplicator.h
+++ b/Engine/Source/Developer/GameplayDebugger/Public/GameplayDebuggerCategoryReplicator.h
@@ -159,7 +159,7 @@ protected:
 	friend FGameplayDebuggerNetPack;
 
 	UPROPERTY(Replicated)
-	APlayerController* OwnerPC;
+	TObjectPtr<APlayerController> OwnerPC;
 
 	UPROPERTY(Replicated)
 	bool bIsEnabled;
@@ -175,7 +175,7 @@ protected:
 
 	/** rendering component needs to attached to some actor, and this is as good as any */
 	UPROPERTY()
-	UGameplayDebuggerRenderingComponent* RenderingComp;
+	TObjectPtr<UGameplayDebuggerRenderingComponent> RenderingComp;
 
 	/** category objects */
 	TArray<TSharedRef<FGameplayDebuggerCategory> > Categories;
