Adding tracking for forced net updates to rep graph csv tracker

diff --git a/Engine/Plugins/Runtime/ReplicationGraph/Source/Private/ReplicationGraph.cpp b/Engine/Plugins/Runtime/ReplicationGraph/Source/Private/ReplicationGraph.cpp
--- a/Engine/Plugins/Runtime/ReplicationGraph/Source/Private/ReplicationGraph.cpp
+++ b/Engine/Plugins/Runtime/ReplicationGraph/Source/Private/ReplicationGraph.cpp
@@ -145,6 +145,7 @@ CSV_DEFINE_CATEGORY(ReplicationGraphKBytes, WITH_SERVER_CODE);
 CSV_DEFINE_CATEGORY(ReplicationGraphChannelsOpened, WITH_SERVER_CODE);
 CSV_DEFINE_CATEGORY(ReplicationGraphNumReps, WITH_SERVER_CODE);
 CSV_DEFINE_CATEGORY(ReplicationGraphVisibleLevels, WITH_SERVER_CODE);
+CSV_DEFINE_CATEGORY(ReplicationGraphForcedUpdates, WITH_SERVER_CODE);
 
 static TAutoConsoleVariable<FString> CVarRepGraphConditionalBreakpointActorName(TEXT("Net.RepGraph.ConditionalBreakpointActorName"), TEXT(""), 
 	TEXT("Helper CVar for debugging. Set this string to conditionally log/breakpoint various points in the repgraph pipeline. Useful for bugs like 'why is this actor channel closing'"), ECVF_Default );
@@ -674,6 +675,8 @@ void UReplicationGraph::ForceNetUpdate(AActor* Actor)
 	{
 		RepInfo->ForceNetUpdateFrame = ReplicationGraphFrame;
 		RepInfo->Events.ForceNetUpdate.Broadcast(Actor, *RepInfo);
+
+		CSVTracker.PostActorForceUpdated(Actor->GetClass());
 	}
 }
 
diff --git a/Engine/Plugins/Runtime/ReplicationGraph/Source/Public/ReplicationGraphTypes.h b/Engine/Plugins/Runtime/ReplicationGraph/Source/Public/ReplicationGraphTypes.h
--- a/Engine/Plugins/Runtime/ReplicationGraph/Source/Public/ReplicationGraphTypes.h
+++ b/Engine/Plugins/Runtime/ReplicationGraph/Source/Public/ReplicationGraphTypes.h
@@ -1457,6 +1457,7 @@ CSV_DECLARE_CATEGORY_EXTERN(ReplicationGraphKBytes);
 CSV_DECLARE_CATEGORY_EXTERN(ReplicationGraphChannelsOpened);
 CSV_DECLARE_CATEGORY_EXTERN(ReplicationGraphNumReps);
 CSV_DECLARE_CATEGORY_EXTERN(ReplicationGraphVisibleLevels);
+CSV_DECLARE_CATEGORY_EXTERN(ReplicationGraphForcedUpdates);
 
 #ifndef REPGRAPH_CSV_TRACKER
 #define REPGRAPH_CSV_TRACKER (CSV_PROFILER && WITH_SERVER_CODE)
@@ -1664,6 +1665,28 @@ struct FReplicationGraphCSVTracker
 #endif
 	}
 
+	void PostActorForceUpdated(UClass* ActorClass)
+	{
+#if REPGRAPH_CSV_TRACKER
+		if (!bIsCapturing)
+		{
+			return;
+		}
+
+		FTrackedData* TrackedData = ExplicitClassTracker.Find(ActorClass);
+		if (TrackedData == nullptr)
+		{
+			TrackedData = ImplicitClassTracker.GetChecked(ActorClass).Get();
+			if (TrackedData == nullptr)
+			{
+				TrackedData = &EverythingElse;
+			}
+		}
+
+		TrackedData->ForcedUpdates++;
+#endif
+	}
+
 	void ResetTrackedClasses()
 	{
 		ExplicitClassTracker.Reset();
@@ -1762,6 +1785,7 @@ private:
 		int64 BitsAccumulated = 0;
 		int32 ChannelsOpened = 0;
 		int32 NumReplications = 0;
+		int32 ForcedUpdates = 0;
 
 		FName StatName;
 
@@ -1771,6 +1795,7 @@ private:
 			BitsAccumulated = 0;
 			ChannelsOpened = 0;
 			NumReplications = 0;
+			ForcedUpdates = 0;
 		}
 	};
 
@@ -1803,6 +1828,7 @@ private:
 		Profiler->RecordCustomStat(Data.StatName, CSV_CATEGORY_INDEX(ReplicationGraphMS), static_cast<float>(Data.CPUTimeAccumulated) * 1000.f, ECsvCustomStatOp::Set);
 		Profiler->RecordCustomStat(Data.StatName, CSV_CATEGORY_INDEX(ReplicationGraphChannelsOpened), static_cast<float>(Data.ChannelsOpened), ECsvCustomStatOp::Set);
 		Profiler->RecordCustomStat(Data.StatName, CSV_CATEGORY_INDEX(ReplicationGraphNumReps), static_cast<float>(Data.NumReplications), ECsvCustomStatOp::Set);
+		Profiler->RecordCustomStat(Data.StatName, CSV_CATEGORY_INDEX(ReplicationGraphForcedUpdates), static_cast<float>(Data.ForcedUpdates), ECsvCustomStatOp::Set);
 		Data.Reset();
 	}
 #endif
