GitHub 7723 : Use DO_ENABLE_REPGRAPH_DEBUG_ACTOR where applicable
#7723

diff --git a/Engine/Plugins/Runtime/ReplicationGraph/Source/Private/ReplicationGraphDebugging.cpp b/Engine/Plugins/Runtime/ReplicationGraph/Source/Private/ReplicationGraphDebugging.cpp
--- a/Engine/Plugins/Runtime/ReplicationGraph/Source/Private/ReplicationGraphDebugging.cpp
+++ b/Engine/Plugins/Runtime/ReplicationGraph/Source/Private/ReplicationGraphDebugging.cpp
@@ -708,7 +708,7 @@ FConsoleCommandWithWorldAndArgsDelegate::CreateLambda([](const TArray<FString>&
 }));
 #endif
 
-#if !(UE_BUILD_SHIPPING | UE_BUILD_TEST)
+#if !(UE_BUILD_SHIPPING || UE_BUILD_TEST)
 FAutoConsoleCommandWithWorldAndArgs NetRepGraphForceRebuild(TEXT("Net.RepGraph.Spatial.ForceRebuild"),TEXT(""),
 	FConsoleCommandWithWorldAndArgsDelegate::CreateLambda([](const TArray<FString>& Args, UWorld* World)
 	{
@@ -1031,7 +1031,7 @@ void UReplicationGraph::CollectRepListStats(FActorRepListStatCollector& StatColl
 	}
 }
 
-#if !(UE_BUILD_SHIPPING || UE_BUILD_TEST)
+#if DO_ENABLE_REPGRAPH_DEBUG_ACTOR
 AReplicationGraphDebugActor* UReplicationGraph::CreateDebugActor() const
 {
 	return GetWorld()->SpawnActor<AReplicationGraphDebugActor>();
