Fix rep graph opening channels for actors that persisted through seamless travel
Fix crash in rep graph debugging post seamless travel

diff --git a/Engine/Plugins/Runtime/ReplicationGraph/Source/Private/ReplicationGraph.cpp b/Engine/Plugins/Runtime/ReplicationGraph/Source/Private/ReplicationGraph.cpp
--- a/Engine/Plugins/Runtime/ReplicationGraph/Source/Private/ReplicationGraph.cpp
+++ b/Engine/Plugins/Runtime/ReplicationGraph/Source/Private/ReplicationGraph.cpp
@@ -2643,7 +2643,8 @@ void UNetReplicationGraphConnection::NotifyResetAllNetworkActors()
 		Node->NotifyResetAllNetworkActors();
 	}
 
-	ActorInfoMap.ResetActorMap();
+	// Intentionally not resetting the actor info map, since it may contain references to open channels 
+	// Open channels associated with actors that are not added back into the graph will be cleaned up normally using the close frame number
 }
 
 FActorRepListRefView& UNetReplicationGraphConnection::GetPrevDormantActorListForNode(const UReplicationGraphNode* GridNode)
