ReplicationGraph
* Allow explicit base names when creating graph nodes.
* This can improve readability when using logging and debug cmds like RepGraph.PrintGraph

diff --git a/Engine/Plugins/Runtime/ReplicationGraph/Source/Public/ReplicationGraph.h b/Engine/Plugins/Runtime/ReplicationGraph/Source/Public/ReplicationGraph.h
--- a/Engine/Plugins/Runtime/ReplicationGraph/Source/Public/ReplicationGraph.h
+++ b/Engine/Plugins/Runtime/ReplicationGraph/Source/Public/ReplicationGraph.h
@@ -122,6 +122,18 @@ public:
 		return NewNode;
 	}
 
+	/** Allocates and initializes ChildNode of a specific type T with a given base name. This is what you will want to call in your FCreateChildNodeFuncs.  */
+	template< class T >
+	T* CreateChildNode(FName NodeBaseName)
+	{
+		FName UniqueNodeName = MakeUniqueObjectName(this, T::StaticClass(), NodeBaseName);
+		T* NewNode = NewObject<T>(this, UniqueNodeName);
+		NewNode->Initialize(GraphGlobals);
+		AllChildNodes.Add(NewNode);
+		return NewNode;
+	}
+
+
 	/** Node removal behavior */
 	enum class NodeOrdering
 	{
@@ -941,6 +953,16 @@ public:
 		return NewNode;
 	}
 
+	/** Creates a new node for the graph with a given base name. This and UReplicationNode::CreateChildNode should be the only things that create the graph node UObjects */
+	template< class T >
+	T* CreateNewNode(FName NodeBaseName)
+	{
+		FName UniqueNodeName = MakeUniqueObjectName(this, T::StaticClass(), NodeBaseName);
+		T* NewNode = NewObject<T>(this, UniqueNodeName);
+		InitNode(NewNode);
+		return NewNode;
+	}
+
 	/** Add a global node to the root that will have a chance to emit actor rep lists for all connections */
 	void AddGlobalGraphNode(UReplicationGraphNode* GraphNode);
 	void RemoveGlobalGraphNode(UReplicationGraphNode* GraphNode);
