Rep Graph : Change CVar_RepGraph_DynamicSpatialFrequency_ForceMaxFreq from REPGRAPH_DEVCVAR_SHIPCONST to a regular cvar to allow test builds to set it.

diff --git a/Engine/Plugins/Runtime/ReplicationGraph/Source/Private/ReplicationGraph.cpp b/Engine/Plugins/Runtime/ReplicationGraph/Source/Private/ReplicationGraph.cpp
--- a/Engine/Plugins/Runtime/ReplicationGraph/Source/Private/ReplicationGraph.cpp
+++ b/Engine/Plugins/Runtime/ReplicationGraph/Source/Private/ReplicationGraph.cpp
@@ -3681,7 +3681,9 @@ void UReplicationGraphNode_DynamicSpatialFrequency::GatherActorListsForConnectio
 }
 
 REPGRAPH_DEVCVAR_SHIPCONST(int32, "Net.RepGraph.DynamicSpatialFrequency.Draw", CVar_RepGraph_DynamicSpatialFrequency_Draw, 0, "");
-REPGRAPH_DEVCVAR_SHIPCONST(int32, "Net.RepGraph.DynamicSpatialFrequency.ForceMaxFreq", CVar_RepGraph_DynamicSpatialFrequency_ForceMaxFreq, 0, "Forces DSF to set max frame replication periods on all actors (1 frame rep periods). 1 = default replication. 2 = fast path. 3 = Both (effectively, default)");
+
+int32 CVar_RepGraph_DynamicSpatialFrequency_ForceMaxFreq = 0;
+static FAutoConsoleVariableRef CVarRepGraphDynamicSpatialFrequencyForceMaxFreq(TEXT("Net.RepGraph.DynamicSpatialFrequency.ForceMaxFreq"), CVar_RepGraph_DynamicSpatialFrequency_ForceMaxFreq, TEXT("Forces DSF to set max frame replication periods on all actors (1 frame rep periods). 1 = default replication. 2 = fast path. 3 = Both (effectively, default)"), ECVF_Default);
 
 FORCEINLINE uint32 CalcDynamicReplicationPeriod(const float FinalPCT, const uint32 MinRepPeriod, const uint32 MaxRepPeriod, uint16& OutReplicationPeriodFrame, uint32& OutNextReplicationFrame, const uint32 LastRepFrameNum, const uint32 FrameNum, bool ForFastPath)
 {
