Wrapping replication graph calls with WITH_SERVER_CODE where appropriate

diff --git a/Engine/Plugins/Runtime/ReplicationGraph/Source/Private/ReplicationGraph.cpp b/Engine/Plugins/Runtime/ReplicationGraph/Source/Private/ReplicationGraph.cpp
--- a/Engine/Plugins/Runtime/ReplicationGraph/Source/Private/ReplicationGraph.cpp
+++ b/Engine/Plugins/Runtime/ReplicationGraph/Source/Private/ReplicationGraph.cpp
@@ -189,6 +189,7 @@ REPGRAPH_DEVCVAR_SHIPCONST(int32, "Net.RepGraph.ActorChannelWarnThreshold", CVar
 TMap<FActorConnectionPair, int32> ActorChannelCreateCounter;
 TMap<FActorConnectionPair, int32> ActorChannelDestroyCounter;
 
+#if WITH_SERVER_CODE
 static TAutoConsoleVariable<FString> CVarRepGraphConditionalPairActorName(TEXT("Net.RepGraph.ConditionalPairName"), TEXT(""), TEXT(""), ECVF_Default );
 void UpdateActorConnectionCounter(AActor* InActor, UNetConnection* InConnection, TMap<FActorConnectionPair, int32>& Counter)
 {
@@ -215,6 +216,7 @@ void UpdateActorConnectionCounter(AActor* InActor, UNetConnection* InConnection,
 	}
 #endif
 }
+#endif // WITH_SERVER_CODE
 
 // --------------------------------------------------------------------------------------------------------------------------------------------
 // --------------------------------------------------------------------------------------------------------------------------------------------
@@ -299,6 +301,7 @@ void UReplicationGraph::Serialize(FArchive& Ar)
 
 void UReplicationGraph::InitForNetDriver(UNetDriver* InNetDriver)
 {
+#if WITH_SERVER_CODE
 	LLM_SCOPE_BYTAG(NetRepGraph);
 
 	NetDriver = InNetDriver;
@@ -310,6 +313,7 @@ void UReplicationGraph::InitForNetDriver(UNetDriver* InNetDriver)
 	{
 		AddClientConnection(ClientConnection);
 	}
+#endif // WITH_SERVER_CODE
 }
 
 void UReplicationGraph::TearDown()
@@ -382,6 +386,7 @@ void UReplicationGraph::RemoveConnectionGraphNode(UReplicationGraphNode* GraphNo
 
 UNetReplicationGraphConnection* UReplicationGraph::FindOrAddConnectionManager(UNetConnection* NetConnection)
 {
+#if WITH_SERVER_CODE
 	check(NetConnection);
 
 	// Children do not have a connection manager, this is handled by their parent.
@@ -417,10 +422,14 @@ UNetReplicationGraphConnection* UReplicationGraph::FindOrAddConnectionManager(UN
 	UNetReplicationGraphConnection* NewManager = CreateClientConnectionManagerInternal(NetConnection);
 	PendingConnections.Add(NewManager);
 	return NewManager;
+#else
+	return nullptr;
+#endif // WITH_SERVER_CODE
 }
 
 void UReplicationGraph::AddClientConnection(UNetConnection* NetConnection)
 {
+#if WITH_SERVER_CODE
 	// Children do not have a connection manager, do not proceed with this function in this case.
 	// Default behavior never calls this function with child connections anyways, so this is really only here for protection.
 	if (NetConnection->GetUChildConnection() != nullptr)
@@ -445,6 +454,7 @@ void UReplicationGraph::AddClientConnection(UNetConnection* NetConnection)
 
 	// Create it
 	Connections.Add(CreateClientConnectionManagerInternal(NetConnection));
+#endif // WITH_SERVER_CODE
 }
 
 UNetReplicationGraphConnection* UReplicationGraph::CreateClientConnectionManagerInternal(UNetConnection* Connection)
@@ -549,6 +559,7 @@ void UReplicationGraph::SetRepDriverWorld(UWorld* InWorld)
 
 void UReplicationGraph::InitializeActorsInWorld(UWorld* InWorld)
 {
+#if WITH_SERVER_CODE
 	check(GraphGlobals.IsValid());
 	checkf(GraphGlobals->World == InWorld, TEXT("UReplicationGraph::InitializeActorsInWorld world mismatch. %s vs %s"), *GetPathNameSafe(GraphGlobals->World), *GetPathNameSafe(InWorld));
 
@@ -567,10 +578,12 @@ void UReplicationGraph::InitializeActorsInWorld(UWorld* InWorld)
 			});
 		}
 	}
+#endif // WITH_SERVER_CODE
 }
 
 void UReplicationGraph::InitializeForWorld(UWorld* World)
 {
+#if WITH_SERVER_CODE
 	LLM_SCOPE_BYTAG(NetRepGraph);
 
 	ActiveNetworkActors.Reset();
@@ -597,6 +610,7 @@ void UReplicationGraph::InitializeForWorld(UWorld* World)
 			}
 		}
 	}
+#endif // WITH_SERVER_CODE
 }
 
 void UReplicationGraph::AddNetworkActor(AActor* Actor)
@@ -922,6 +936,7 @@ int32 UReplicationGraph::ServerReplicateActors(float DeltaSeconds)
 		}
 	};
 
+#if WITH_SERVER_CODE
 	// -------------------------------------------------------
 	//	PREPARE (Global)
 	// -------------------------------------------------------
@@ -1171,11 +1186,13 @@ int32 UReplicationGraph::ServerReplicateActors(float DeltaSeconds)
 	PostServerReplicateStats(FrameReplicationStats);
 
 	CSVTracker.EndReplicationFrame();
+#endif // WITH_SERVER_CODE
 	return 0;
 }
 
 void UReplicationGraph::ReplicateActorListsForConnections_Default(UNetReplicationGraphConnection* ConnectionManager, FGatheredReplicationActorLists& GatheredReplicationListsForConnection, FNetViewerArray& Viewers)
 {
+#if WITH_SERVER_CODE
 #if !(UE_BUILD_SHIPPING || UE_BUILD_TEST)
 	const bool bEnableFullActorPrioritizationDetails = DO_REPGRAPH_DETAILS(bEnableFullActorPrioritizationDetailsAllConnections || ConnectionManager->bEnableFullActorPrioritizationDetails);
 	const bool bDoDistanceCull = (CVar_RepGraph_SkipDistanceCull == 0);
@@ -1426,10 +1443,12 @@ void UReplicationGraph::ReplicateActorListsForConnections_Default(UNetReplicatio
 		UE_LOG(LogReplicationGraph, Display, TEXT("Gathered Lists: %d Gathered Actors: %d  PrioritizedActors: %d"), NumGatheredListsOnConnection, NumGatheredActorsOnConnection, NumPrioritizedActorsOnConnection);
 		UE_LOG(LogReplicationGraph, Display, TEXT("Connection Loaded Streaming Levels: %d"), NetConnection->ClientVisibleLevelNames.Num());
 	}
+#endif // WITH_SERVER_CODE
 }
 
 void UReplicationGraph::ReplicateActorsForConnection(UNetConnection* NetConnection, FPerConnectionActorInfoMap& ConnectionActorInfoMap, UNetReplicationGraphConnection* ConnectionManager, const uint32 FrameNum)
 {
+#if WITH_SERVER_CODE
 	QUICK_SCOPE_CYCLE_COUNTER(NET_ReplicateActors_ReplicateActorsForConnection);
 
 	for (int32 ActorIdx = 0; ActorIdx < PrioritizedReplicationList.Items.Num(); ++ActorIdx)
@@ -1463,6 +1482,7 @@ void UReplicationGraph::ReplicateActorsForConnection(UNetConnection* NetConnecti
 			break;
 		}
 	}
+#endif // WITH_SERVER_CODE
 }
 
 struct FScopedQueuedBits
@@ -1516,6 +1536,7 @@ struct FScopedFastPathTracker
 
 void UReplicationGraph::ReplicateActorListsForConnections_FastShared(UNetReplicationGraphConnection* ConnectionManager, FGatheredReplicationActorLists& GatheredReplicationListsForConnection, FNetViewerArray& Viewers)
 {
+#if WITH_SERVER_CODE
 	if (CVar_RepGraph_EnableFastSharedPath == false)
 	{
 		return;
@@ -1630,12 +1651,14 @@ void UReplicationGraph::ReplicateActorListsForConnections_FastShared(UNetReplica
 			}
 		}
 	}
+#endif // WITH_SERVER_CODE
 }
 
 REPGRAPH_DEVCVAR_SHIPCONST(int32, "Net.RepGraph.FastShared.ForceFull", CVar_RepGraph_FastShared_ForceFull, 0, "Redirects calls to ReplicateSingleActor_FastShared to ReplicateSingleActor");
 
 int64 UReplicationGraph::ReplicateSingleActor_FastShared(AActor* Actor, FConnectionReplicationActorInfo& ConnectionData, FGlobalActorReplicationInfo& GlobalActorInfo, UNetReplicationGraphConnection& ConnectionManager, const uint32 FrameNum)
 {
+#if WITH_SERVER_CODE
 	UNetConnection* NetConnection = ConnectionManager.NetConnection;
 
 	// No matter what we consider this FastShared rep to happen. Even if the actor doesn't produce a bunch or its empty or stale, etc. We still consider this replication to have happened
@@ -1763,10 +1786,14 @@ int64 UReplicationGraph::ReplicateSingleActor_FastShared(AActor* Actor, FConnect
 	ensureAlwaysMsgf(OutBunch.bHasMustBeMappedGUIDs == 0, TEXT("FastShared bHasMustBeMappedGUIDs! %s"), *Actor->GetPathName());
 
 	return OutBunch.GetNumBits();
+#else
+	return 0;
+#endif // WITH_SERVER_CODE
 }
 
 int64 UReplicationGraph::ReplicateSingleActor(AActor* Actor, FConnectionReplicationActorInfo& ActorInfo, FGlobalActorReplicationInfo& GlobalActorInfo, FPerConnectionActorInfoMap& ConnectionActorInfoMap, UNetReplicationGraphConnection& ConnectionManager, const uint32 FrameNum)
 {
+#if WITH_SERVER_CODE
 	RG_QUICK_SCOPE_CYCLE_COUNTER(NET_ReplicateActors_ReplicateSingleActor);
 
 	// These checks will happen anyway in UActorChannel::ReplicateActor, but we need to be able to detect them to prevent crashes.
@@ -1947,6 +1974,9 @@ int64 UReplicationGraph::ReplicateSingleActor(AActor* Actor, FConnectionReplicat
 	}
 
 	return BitsWritten;
+#else
+	return 0;
+#endif //WITH_SERVER_CODE
 }
 
 void UReplicationGraph::HandleStarvedActorList(const FPrioritizedRepList& List, int32 StartIdx, FPerConnectionActorInfoMap& ConnectionActorInfoMap, uint32 FrameNum)
@@ -1990,6 +2020,7 @@ void UReplicationGraph::UpdateActorChannelCloseFrameNum(AActor* Actor, FConnecti
 
 bool UReplicationGraph::ProcessRemoteFunction(class AActor* Actor, UFunction* Function, void* Parameters, FOutParmRec* OutParms, FFrame* Stack, class UObject* SubObject )
 {
+#if WITH_SERVER_CODE
 	// ----------------------------------
 	// Setup
 	// ----------------------------------
@@ -2253,6 +2284,7 @@ bool UReplicationGraph::ProcessRemoteFunction(class AActor* Actor, UFunction* Fu
 	{
 		UE_LOG(LogNet, Warning, TEXT("UReplicationGraph::ProcessRemoteFunction: No owning connection for actor %s. Function %s will not be processed."), *Actor->GetName(), *Function->GetName());
 	}
+#endif // WITH_SERVER_CODE
 
 	// return true because we don't want the net driver to do anything else
 	return true;
@@ -2412,6 +2444,7 @@ void UNetReplicationGraphConnection::TearDown()
 }
 void UNetReplicationGraphConnection::NotifyActorChannelAdded(AActor* Actor, class UActorChannel* Channel)
 {
+#if WITH_SERVER_CODE
 	UpdateActorConnectionCounter(Actor, Channel->Connection, ActorChannelCreateCounter);
 
 	if (RepGraphConditionalActorBreakpoint(Actor, Channel->Connection))
@@ -2435,6 +2468,7 @@ void UNetReplicationGraphConnection::NotifyActorChannelAdded(AActor* Actor, clas
 
 	ActorInfo.Channel = Channel;
 	ActorInfoMap.AddChannel(Actor, Channel);
+#endif // WITH_SERVER_CODE
 }
 
 void UNetReplicationGraphConnection::NotifyActorChannelRemoved(AActor* Actor)
@@ -2542,6 +2576,7 @@ bool UNetReplicationGraphConnection::PrepareForReplication()
 
 void UNetReplicationGraphConnection::NotifyAddDestructionInfo(FActorDestructionInfo* DestructInfo)
 {
+#if WITH_SERVER_CODE
 	if (DestructInfo->StreamingLevelName != NAME_None)
 	{
 		if (NetConnection->ClientVisibleLevelNames.Contains(DestructInfo->StreamingLevelName) == false)
@@ -2570,10 +2605,12 @@ void UNetReplicationGraphConnection::NotifyAddDestructionInfo(FActorDestructionI
 
 	PendingDestructInfoList.Emplace( FCachedDestructInfo(DestructInfo) );
 	//UE_LOG(LogReplicationGraph, Display, TEXT("::NotifyAddDestructionInfo. Connection: %s. DestructInfo: %s. NewTotal: %d"), *NetConnection->Describe(), *DestructInfo->PathName, PendingDestructInfoList.Num());
+#endif // WITH_SERVER_CODE
 }
 
 void UNetReplicationGraphConnection::NotifyAddDormantDestructionInfo(AActor* Actor)
 {
+#if WITH_SERVER_CODE
 	if (Actor && NetConnection && NetConnection->Driver && NetConnection->Driver->GuidCache)
 	{
 		ULevel* Level = Actor->GetLevel();
@@ -2608,6 +2645,7 @@ void UNetReplicationGraphConnection::NotifyAddDormantDestructionInfo(AActor* Act
 			Info.PathName = Actor->GetName();
 		}
 	}
+#endif // WITH_SERVER_CODE
 }
 
 void UNetReplicationGraphConnection::NotifyRemoveDestructionInfo(FActorDestructionInfo* DestructInfo)
@@ -2717,6 +2755,7 @@ int64 UNetReplicationGraphConnection::ReplicateDestructionInfos(const FNetViewer
 	QUICK_SCOPE_CYCLE_COUNTER(ReplicateDestructionInfos);
 
 	int64 NumBits = 0;
+#if WITH_SERVER_CODE
 	for (int32 idx=PendingDestructInfoList.Num()-1; idx >=0; --idx)
 	{
 		FCachedDestructInfo& Info = PendingDestructInfoList[idx];
@@ -2755,7 +2794,7 @@ int64 UNetReplicationGraphConnection::ReplicateDestructionInfos(const FNetViewer
 			PendingDestructInfoList.RemoveAtSwap(idx, 1, false);
 		}
 	}
-
+#endif // #if WITH_SERVER_CODE
 	return NumBits;
 }
 
@@ -2766,6 +2805,8 @@ int64 UNetReplicationGraphConnection::ReplicateDormantDestructionInfos()
 
 	int64 NumBits = 0;
 
+#if WITH_SERVER_CODE
+
 	if (NetConnection && NetConnection->Driver)
 	{
 		for (const FCachedDormantDestructInfo& Info : PendingDormantDestructList)
@@ -2785,6 +2826,7 @@ int64 UNetReplicationGraphConnection::ReplicateDormantDestructionInfos()
 		PendingDormantDestructList.Reset();
 		TrackedDormantDestructionInfos.Reset();
 	}
+#endif // WITH_SERVER_CODE
 
 	return NumBits;
 }
@@ -3291,6 +3333,7 @@ void UReplicationGraphNode_ActorListFrequencyBuckets::NotifyResetAllNetworkActor
 
 void UReplicationGraphNode_ActorListFrequencyBuckets::GatherActorListsForConnection(const FConnectionGatherActorListParameters& Params)
 {
+#if WITH_SERVER_CODE
 	const FSettings& MySettings = GetSettings();
 	const uint32 ReplicationFrameNum = Params.ReplicationFrameNum;
 	if (MySettings.EnableFastPath)
@@ -3322,6 +3365,7 @@ void UReplicationGraphNode_ActorListFrequencyBuckets::GatherActorListsForConnect
 	}
 
 	StreamingLevelCollection.Gather(Params);
+#endif // WITH_SERVER_CODE
 }
 
 void UReplicationGraphNode_ActorListFrequencyBuckets::SetNonStreamingCollectionSize(const int32 NewSize)
@@ -3495,6 +3539,7 @@ FORCEINLINE bool ReplicatesEveryFrame(const FConnectionReplicationActorInfo& Con
 
 void UReplicationGraphNode_DynamicSpatialFrequency::GatherActorListsForConnection(const FConnectionGatherActorListParameters& Params)
 {
+#if WITH_SERVER_CODE
 	repCheck(GraphGlobals.IsValid());
 
 	UReplicationGraph* RepGraph = GraphGlobals->ReplicationGraph;
@@ -3678,6 +3723,7 @@ void UReplicationGraphNode_DynamicSpatialFrequency::GatherActorListsForConnectio
 			UE_LOG(LogReplicationGraph, Display, TEXT("Uncapped bandwidth usage of UReplicationGraphNode_DynamicSpatialFrequency = %d bits -> %d bytes -> %.2f KBytes/sec"), BitsWritten, (BitsWritten+7) >> 3,  ((float)((BitsWritten+7)>>3)/1024.f) * GraphGlobals->ReplicationGraph->NetDriver->NetServerMaxTickRate);
 		}
 	}
+#endif // WITH_SERVER_CODE
 }
 
 REPGRAPH_DEVCVAR_SHIPCONST(int32, "Net.RepGraph.DynamicSpatialFrequency.Draw", CVar_RepGraph_DynamicSpatialFrequency_Draw, 0, "");
@@ -3733,6 +3779,7 @@ void UReplicationGraphNode_DynamicSpatialFrequency::CalcFrequencyForActor(AActor
 
 void UReplicationGraphNode_DynamicSpatialFrequency::CalcFrequencyForActor(AActor* Actor, UReplicationGraph* RepGraph, UNetConnection* NetConnection, FGlobalActorReplicationInfo& GlobalInfo, FPerConnectionActorInfoMap& ConnectionMap, FSettings& MySettings, const FNetViewerArray& Viewers, const uint32 FrameNum, int32 ExistingItemIndex)
 {
+#if WITH_SERVER_CODE
 	FConnectionReplicationActorInfo& ConnectionInfo = ConnectionMap.FindOrAdd(Actor);
 
 	// If we need to filter out the actor and he is already in the SortedReplicationList, we need to remove it (instead of just skipping/returning).
@@ -3922,6 +3969,7 @@ void UReplicationGraphNode_DynamicSpatialFrequency::CalcFrequencyForActor(AActor
 	// No zone was found. This is bad
 	UE_LOG(LogReplicationGraph, Warning, TEXT("UReplicationGraphNode_DynamicSpatialFrequency::CalcFrequencyForActor: %s was not placed in any valid zone. Viewer: %s DotP: %.2f "), *Actor->GetName(), ((LowestDistanceViewer->Connection != nullptr) ? *LowestDistanceViewer->Connection->GetName() : TEXT("INVALID")), DotP);
 	RemoveExistingItem();
+#endif // WITH_SERVER_CODE
 }
 
 void UReplicationGraphNode_DynamicSpatialFrequency::GatherActors(const FActorRepListRefView& RepList, FGlobalActorReplicationInfoMap& GlobalMap, FPerConnectionActorInfoMap& ConnectionMap, const FConnectionGatherActorListParameters& Params, UNetConnection* NetConnection)
@@ -4009,6 +4057,7 @@ void UReplicationGraphNode_ConnectionDormancyNode::TearDown()
 
 void UReplicationGraphNode_ConnectionDormancyNode::GatherActorListsForConnection(const FConnectionGatherActorListParameters& Params)
 {
+#if WITH_SERVER_CODE
 	RG_QUICK_SCOPE_CYCLE_COUNTER(RepGraphNode_ConnectionDormancy_Gather);
 
 	LastGatheredFrame = Params.ReplicationFrameNum;
@@ -4040,10 +4089,12 @@ void UReplicationGraphNode_ConnectionDormancyNode::GatherActorListsForConnection
 			UE_LOG(LogReplicationGraph, Verbose, TEXT("Level Not Loaded %s. (Client has %d levels loaded)"), *StreamingList.StreamingLevelName.ToString(), Params.ClientVisibleLevelNamesRef.Num());
 		}
 	}
+#endif // WITH_SERVER_CODE
 }
 
 void UReplicationGraphNode_ConnectionDormancyNode::ConditionalGatherDormantActorsForConnection(FActorRepListRefView& ConnectionList, const FConnectionGatherActorListParameters& Params, FActorRepListRefView* RemovedList)
 {
+#if WITH_SERVER_CODE
 	FPerConnectionActorInfoMap& ConnectionActorInfoMap = Params.ConnectionManager.ActorInfoMap;
 	FGlobalActorReplicationInfoMap* GlobalActorReplicationInfoMap = GraphGlobals->GlobalActorReplicationInfoMap;
 
@@ -4089,6 +4140,7 @@ void UReplicationGraphNode_ConnectionDormancyNode::ConditionalGatherDormantActor
 			TrickleStartCounter--;		
 		}
 	}
+#endif // WITH_SERVER_CODE
 }
 
 bool ContainsReverse(const FActorRepListRefView& List, FActorRepListType Actor)
@@ -4305,6 +4357,7 @@ void UReplicationGraphNode_DormancyNode::RemoveDormantActor(const FNewReplicated
 
 void UReplicationGraphNode_DormancyNode::GatherActorListsForConnection(const FConnectionGatherActorListParameters& Params)
 {
+#if WITH_SERVER_CODE
 	int32 NumViewersAboveMaxZ = 0;
 	for (const FNetViewer& CurViewer : Params.Viewers)
 	{
@@ -4332,7 +4385,7 @@ void UReplicationGraphNode_DormancyNode::GatherActorListsForConnection(const FCo
 		ConnectionNode = CreateConnectionNode(Params);
 		ConnectionNode->GatherActorListsForConnection(Params);
 	}
-	
+#endif // WITH_SERVER_CODE
 }
 
 UReplicationGraphNode_ConnectionDormancyNode* UReplicationGraphNode_DormancyNode::GetExistingConnectionNode(const FConnectionGatherActorListParameters& Params)
@@ -4993,6 +5046,7 @@ static FAutoConsoleVariableRef CVarRepGraphSpatialBiasCreep(TEXT("Net.RepGraph.S
 
 void UReplicationGraphNode_GridSpatialization2D::PrepareForReplication()
 {
+#if	WITH_SERVER_CODE
 	RG_QUICK_SCOPE_CYCLE_COUNTER(UReplicationGraphNode_GridSpatialization2D_PrepareForReplication);
 
 	FGlobalActorReplicationInfoMap* GlobalRepMap = GraphGlobals.IsValid() ? GraphGlobals->GlobalActorReplicationInfoMap : nullptr;
@@ -5391,6 +5445,7 @@ void UReplicationGraphNode_GridSpatialization2D::PrepareForReplication()
 
 		bNeedsRebuild = false;
 	}
+#endif // WITH_SERVER_CODE
 }
 
 // Small structure to make it easier to keep track of 
@@ -5408,6 +5463,7 @@ struct FPlayerGridCellInformation
 
 void UReplicationGraphNode_GridSpatialization2D::GatherActorListsForConnection(const FConnectionGatherActorListParameters& Params)
 {
+#if WITH_SERVER_CODE
 	TArray<FLastLocationGatherInfo>& LastLocationArray = Params.ConnectionManager.LastGatherLocations;
 	TArray<FVector2D, FReplicationGraphConnectionsAllocator> UniqueCurrentLocations;
 
@@ -5610,6 +5666,7 @@ void UReplicationGraphNode_GridSpatialization2D::GatherActorListsForConnection(c
 			}
 		}
 	}
+#endif // WITH_SERVER_CODE
 }
 
 void UReplicationGraphNode_GridSpatialization2D::NotifyActorCullDistChange(AActor* Actor, FGlobalActorReplicationInfo& GlobalInfo, float OldDist)
@@ -5741,6 +5798,8 @@ void UReplicationGraphNode_AlwaysRelevant::GatherActorListsForConnection(const F
 
 void UReplicationGraphNode_TearOff_ForConnection::GatherActorListsForConnection(const FConnectionGatherActorListParameters& Params)
 {
+#if WITH_SERVER_CODE
+
 	if (TearOffActors.Num() > 0)
 	{
 		ReplicationActorList.Reset();
@@ -5786,6 +5845,7 @@ void UReplicationGraphNode_TearOff_ForConnection::GatherActorListsForConnection(
 			Params.OutGatheredReplicationLists.AddReplicationActorList(ReplicationActorList);
 		}
 	}
+#endif // WITH_SERVER_CODE
 }
 
 void UReplicationGraphNode_TearOff_ForConnection::NotifyTearOffActor(AActor* Actor, uint32 FrameNum)
@@ -5800,6 +5860,7 @@ void UReplicationGraphNode_AlwaysRelevant_ForConnection::GatherActorListsForConn
 	// Call super to add any actors that were explicitly given to use via NotifyAddNetworkActor
 	Super::GatherActorListsForConnection(Params);
 
+#if WITH_SERVER_CODE
 	auto UpdateActor = [&](AActor* NewActor, AActor*& LastActor)
 	{
 		if (NewActor != LastActor)
@@ -5867,6 +5928,7 @@ void UReplicationGraphNode_AlwaysRelevant_ForConnection::GatherActorListsForConn
 	{
 		Params.OutGatheredReplicationLists.AddReplicationActorList(ReplicationActorList);
 	}
+#endif // WITH_SERVER_CODE
 }
 
 void UReplicationGraphNode_AlwaysRelevant_ForConnection::TearDown()
@@ -5884,6 +5946,7 @@ void UReplicationGraphNode_AlwaysRelevant_ForConnection::OnCollectActorRepListSt
 }
 
 // -------------------------------------------------------
+#if WITH_SERVER_CODE
 
 FAutoConsoleCommandWithWorldAndArgs NetRepGraphPrintChannelCounters(TEXT("Net.RepGraph.PrintActorChannelCounters"),TEXT(""),
 	FConsoleCommandWithWorldAndArgsDelegate::CreateLambda([](const TArray<FString>& Args, UWorld* World)
@@ -5912,3 +5975,5 @@ FAutoConsoleCommandWithWorldAndArgs ChangeActorDiscoveryBudget(TEXT("Net.RepGrap
 		It->SetActorDiscoveryBudget(BudgetInKBPS);
 	}
 }));
+
+#endif // WITH_SERVER_CODE
\ No newline at end of file
