Adding option to remove unused forcenetupdate delegate

diff --git a/Engine/Plugins/Runtime/ReplicationGraph/Source/Private/ReplicationGraph.cpp b/Engine/Plugins/Runtime/ReplicationGraph/Source/Private/ReplicationGraph.cpp
--- a/Engine/Plugins/Runtime/ReplicationGraph/Source/Private/ReplicationGraph.cpp
+++ b/Engine/Plugins/Runtime/ReplicationGraph/Source/Private/ReplicationGraph.cpp
@@ -705,7 +705,9 @@ void UReplicationGraph::ForceNetUpdate(AActor* Actor)
 	if (FGlobalActorReplicationInfo* RepInfo = GlobalActorReplicationInfoMap.Find(Actor))
 	{
 		RepInfo->ForceNetUpdateFrame = ReplicationGraphFrame;
+#if REPGRAPH_ENABLE_FORCENETUPDATE_DELEGATE
 		RepInfo->Events.ForceNetUpdate.Broadcast(Actor, *RepInfo);
+#endif // REPGRAPH_ENABLE_FORCENETUPDATE_DELEGATE
 
 		CSVTracker.PostActorForceUpdated(Actor->GetClass());
 	}
diff --git a/Engine/Plugins/Runtime/ReplicationGraph/Source/Public/ReplicationGraphTypes.h b/Engine/Plugins/Runtime/ReplicationGraph/Source/Public/ReplicationGraphTypes.h
--- a/Engine/Plugins/Runtime/ReplicationGraph/Source/Public/ReplicationGraphTypes.h
+++ b/Engine/Plugins/Runtime/ReplicationGraph/Source/Public/ReplicationGraphTypes.h
@@ -56,6 +56,10 @@ LLM_DECLARE_TAG_API(NetRepGraph, REPLICATIONGRAPH_API);
 	const Type Var = Value;
 #endif
 
+#ifndef REPGRAPH_ENABLE_FORCENETUPDATE_DELEGATE
+#define REPGRAPH_ENABLE_FORCENETUPDATE_DELEGATE 1
+#endif // REPGRAPH_ENABLE_FORCENETUPDATE_DELEGATE
+
 // --------------------------------------------------------------------------------------------------------------------------------------------
 // --------------------------------------------------------------------------------------------------------------------------------------------
 // Actor Replication List Types
@@ -527,13 +531,18 @@ struct FFastSharedReplicationInfo
 
 DECLARE_MULTICAST_DELEGATE_FourParams(FNotifyActorChangeDormancy, FActorRepListType, FGlobalActorReplicationInfo&, ENetDormancy /*NewVlue*/, ENetDormancy /*OldValue*/);
 DECLARE_MULTICAST_DELEGATE_TwoParams(FNotifyActorFlushDormancy, FActorRepListType, FGlobalActorReplicationInfo&);
+
+#if REPGRAPH_ENABLE_FORCENETUPDATE_DELEGATE
 DECLARE_MULTICAST_DELEGATE_TwoParams(FNotifyActorForceNetUpdate, FActorRepListType, FGlobalActorReplicationInfo&);
+#endif // REPGRAPH_ENABLE_FORCENETUPDATE_DELEGATE
 
 struct FGlobalActorReplicationEvents
 {
 	FNotifyActorChangeDormancy	DormancyChange;
 	FNotifyActorFlushDormancy	DormancyFlush; // This delegate is cleared after broadcasting
+#if REPGRAPH_ENABLE_FORCENETUPDATE_DELEGATE
 	FNotifyActorForceNetUpdate	ForceNetUpdate;
+#endif // REPGRAPH_ENABLE_FORCENETUPDATE_DELEGATE
 };
 
 /** Per-Actor data that is global for the entire Replication Graph */
