Additional granular mem tracking for more rep graph nodes
Adding scope count to the granular logging

diff --git a/Engine/Plugins/Runtime/ReplicationGraph/Source/Private/ReplicationGraph.cpp b/Engine/Plugins/Runtime/ReplicationGraph/Source/Private/ReplicationGraph.cpp
--- a/Engine/Plugins/Runtime/ReplicationGraph/Source/Private/ReplicationGraph.cpp
+++ b/Engine/Plugins/Runtime/ReplicationGraph/Source/Private/ReplicationGraph.cpp
@@ -2908,6 +2908,18 @@ UReplicationGraphNode::UReplicationGraphNode()
 
 }
 
+void UReplicationGraphNode::Serialize(FArchive& Ar)
+{
+	Super::Serialize(Ar);
+
+	if (Ar.IsCountingMemory())
+	{
+		GRANULAR_NETWORK_MEMORY_TRACKING_INIT(Ar, "UReplicationGraphNode::Serialize");
+
+		GRANULAR_NETWORK_MEMORY_TRACKING_TRACK("AllChildNodes", AllChildNodes.CountBytes(Ar));
+	}
+}
+
 void UReplicationGraphNode::NotifyResetAllNetworkActors()
 {
 	for (UReplicationGraphNode* ChildNode : AllChildNodes)
@@ -3256,6 +3268,29 @@ void UReplicationGraphNode_ActorList::LogActorList(FReplicationGraphDebugInfo& D
 
 UReplicationGraphNode_ActorListFrequencyBuckets::FSettings UReplicationGraphNode_ActorListFrequencyBuckets::DefaultSettings;
 
+void UReplicationGraphNode_ActorListFrequencyBuckets::Serialize(FArchive& Ar)
+{
+	Super::Serialize(Ar);
+
+	if (Ar.IsCountingMemory())
+	{
+		GRANULAR_NETWORK_MEMORY_TRACKING_INIT(Ar, "UReplicationGraphNode_ActorListFrequencyBuckets::Serialize");
+		if (Settings.IsValid())
+		{
+			GRANULAR_NETWORK_MEMORY_TRACKING_TRACK("Settings", Settings->CountBytes(Ar));
+		}
+
+		GRANULAR_NETWORK_MEMORY_TRACKING_TRACK("NonStreamingCollection",
+			NonStreamingCollection.CountBytes(Ar);
+			for (const FActorRepListRefView& RepList : NonStreamingCollection)
+			{
+				RepList.CountBytes(Ar);
+			}
+		);
+		GRANULAR_NETWORK_MEMORY_TRACKING_TRACK("StreamingLevelCollection", StreamingLevelCollection.CountBytes(Ar));
+	}
+}
+
 void UReplicationGraphNode_ActorListFrequencyBuckets::NotifyAddNetworkActor(const FNewReplicatedActorInfo& ActorInfo)
 {
 	UE_CLOG(CVar_RepGraph_LogActorAdd>0, LogReplicationGraph, Display, TEXT("UReplicationGraphNode_ActorListFrequencyBuckets::NotifyAddNetworkActor %s on %s."), *ActorInfo.Actor->GetFullName(), *GetPathName());
@@ -4672,6 +4707,23 @@ UReplicationGraphNode_GridSpatialization2D::UReplicationGraphNode_GridSpatializa
 	bDestroyDormantDynamicActors = CVar_RepGraph_GridSpatialization2D_DestroyDormantDynamicActorsDefault != 0;
 }
 
+void UReplicationGraphNode_GridSpatialization2D::Serialize(FArchive& Ar)
+{
+	Super::Serialize(Ar);
+
+	if (Ar.IsCountingMemory())
+	{
+		GRANULAR_NETWORK_MEMORY_TRACKING_INIT(Ar, "UReplicationGraphNode_GridSpatialization2D::Serialize");
+
+		GRANULAR_NETWORK_MEMORY_TRACKING_TRACK("ClassRebuildDenyList", ClassRebuildDenyList.CountBytes(Ar));
+		GRANULAR_NETWORK_MEMORY_TRACKING_TRACK("DynamicSpatializedActors", DynamicSpatializedActors.CountBytes(Ar));
+		GRANULAR_NETWORK_MEMORY_TRACKING_TRACK("StaticSpatializedActors", StaticSpatializedActors.CountBytes(Ar));
+		GRANULAR_NETWORK_MEMORY_TRACKING_TRACK("PendingStaticSpatializedActors", PendingStaticSpatializedActors.CountBytes(Ar));
+		GRANULAR_NETWORK_MEMORY_TRACKING_TRACK("Grid", Grid.CountBytes(Ar));
+		GRANULAR_NETWORK_MEMORY_TRACKING_TRACK("GatheredNodes", GatheredNodes.CountBytes(Ar));
+	}
+}
+
 void UReplicationGraphNode_GridSpatialization2D::NotifyAddNetworkActor(const FNewReplicatedActorInfo& ActorInfo)
 {
 	ensureAlwaysMsgf(false, TEXT("UReplicationGraphNode_GridSpatialization2D::NotifyAddNetworkActor should not be called directly"));
@@ -5811,6 +5863,19 @@ void UReplicationGraphNode_AlwaysRelevant::GatherActorListsForConnection(const F
 
 // -------------------------------------------------------
 
+void UReplicationGraphNode_TearOff_ForConnection::Serialize(FArchive& Ar)
+{
+	Super::Serialize(Ar);
+
+	if (Ar.IsCountingMemory())
+	{
+		GRANULAR_NETWORK_MEMORY_TRACKING_INIT(Ar, "UReplicationGraphNode_TearOff_ForConnection::Serialize");
+
+		GRANULAR_NETWORK_MEMORY_TRACKING_TRACK("TearOffActors", TearOffActors.CountBytes(Ar));
+		GRANULAR_NETWORK_MEMORY_TRACKING_TRACK("ReplicationActorList", ReplicationActorList.CountBytes(Ar));
+	}
+}
+
 void UReplicationGraphNode_TearOff_ForConnection::GatherActorListsForConnection(const FConnectionGatherActorListParameters& Params)
 {
 #if WITH_SERVER_CODE
diff --git a/Engine/Plugins/Runtime/ReplicationGraph/Source/Public/ReplicationGraph.h b/Engine/Plugins/Runtime/ReplicationGraph/Source/Public/ReplicationGraph.h
--- a/Engine/Plugins/Runtime/ReplicationGraph/Source/Public/ReplicationGraph.h
+++ b/Engine/Plugins/Runtime/ReplicationGraph/Source/Public/ReplicationGraph.h
@@ -78,6 +78,8 @@ public:
 
 	UReplicationGraphNode();
 
+	virtual void Serialize(FArchive& Ar) override;
+
 	/** Called when a network actor is spawned or an actor changes replication status */
 	virtual void NotifyAddNetworkActor(const FNewReplicatedActorInfo& Actor ) PURE_VIRTUAL(UReplicationGraphNode::NotifyAddNetworkActor, );
 	
@@ -281,6 +283,11 @@ public:
 		};
 
 		TArray<FBucketThresholds, TInlineAllocator<4>> BucketThresholds;
+
+		void CountBytes(FArchive& Ar) const 
+		{
+			BucketThresholds.CountBytes(Ar);
+		}
 	};
 
 	/** Default settings for all nodes. By being static, this allows games to easily override the settings are all nodes without having to subclass every graph node class */
@@ -293,6 +300,8 @@ public:
 	
 	UReplicationGraphNode_ActorListFrequencyBuckets() { if (!HasAnyFlags(RF_ClassDefaultObject)) { SetNonStreamingCollectionSize(GetSettings().NumBuckets); } }
 
+	virtual void Serialize(FArchive& Ar) override;
+
 	virtual void NotifyAddNetworkActor(const FNewReplicatedActorInfo& ActorInfo) override;
 	
 	virtual bool NotifyRemoveNetworkActor(const FNewReplicatedActorInfo& ActorInfo, bool bWarnIfNotFound=true) override;
@@ -588,6 +597,8 @@ public:
 
 	UReplicationGraphNode_GridSpatialization2D();
 
+	virtual void Serialize(FArchive& Ar) override;
+
 	virtual void NotifyAddNetworkActor(const FNewReplicatedActorInfo& Actor) override;	
 	virtual bool NotifyRemoveNetworkActor(const FNewReplicatedActorInfo& ActorInfo, bool bWarnIfNotFound=true) override;
 	virtual void NotifyResetAllNetworkActors() override;
@@ -857,6 +868,7 @@ class REPLICATIONGRAPH_API UReplicationGraphNode_TearOff_ForConnection : public
 	GENERATED_BODY()
 
 public:
+	virtual void Serialize(FArchive& Ar) override;
 
 	virtual void NotifyAddNetworkActor(const FNewReplicatedActorInfo& ActorInfo) override { }
 	virtual bool NotifyRemoveNetworkActor(const FNewReplicatedActorInfo& ActorInfo, bool bWarnIfNotFound=true) override { return false; }
diff --git a/Engine/Source/Runtime/Engine/Private/Net/NetworkGranularMemoryLogging.cpp b/Engine/Source/Runtime/Engine/Private/Net/NetworkGranularMemoryLogging.cpp
--- a/Engine/Source/Runtime/Engine/Private/Net/NetworkGranularMemoryLogging.cpp
+++ b/Engine/Source/Runtime/Engine/Private/Net/NetworkGranularMemoryLogging.cpp
@@ -49,16 +49,19 @@ namespace GranularNetworkMemoryTrackingPrivate
 			};
 
 			TArray<FField> FieldValues;
-			FieldValues.Reserve(Fields.Num() + 1);
+			FieldValues.Reserve(Fields.Num() + 2);
 			FieldValues.Add({ &NewPrefix, TotalBytes });
 
+			static const FString CountName = TEXT("Count");
+			FieldValues.Add({ &CountName, Count });
+
 			for (auto ConstIt = Fields.CreateConstIterator(); ConstIt; ++ConstIt)
 			{
 				FieldValues.Add({ &ConstIt.Key(), ConstIt.Value() });
 			}
 
 			// Keep our total scope data in place, but sort the rest of the fields.
-			Sort(FieldValues.GetData() + 1, FieldValues.Num() - 1, FieldLessThan);
+			Sort(FieldValues.GetData() + 2, FieldValues.Num() - 2, FieldLessThan);
 
 			FString ReportRow = FString::Printf(TEXT("%s\r\n%s"),
 				*FString::JoinBy(FieldValues, TEXT(","), [](const FField& Field) { return *Field.Name; }),
@@ -89,12 +92,18 @@ namespace GranularNetworkMemoryTrackingPrivate
 			return TotalBytes;
 		}
 
+		void IncrementCount()
+		{
+			++Count;
+		}
+
 		const FString ScopeName;
 
 	private:
 
 		TMap<FString, TUniquePtr<FNetworkMemoryTrackingScope>> SubScopes;
 		TMap<FString, uint64> Fields;
+		uint32 Count = 0;
 		uint64 TotalBytes = 0u;
 	};
 
@@ -124,9 +133,11 @@ namespace GranularNetworkMemoryTrackingPrivate
 			}
 			else
 			{
-				 CurrentScope = TopLevelScopes.Add(ScopeName, MakeUnique<FNetworkMemoryTrackingScope>(ScopeName)).Get();
+				CurrentScope = TopLevelScopes.Add(ScopeName, MakeUnique<FNetworkMemoryTrackingScope>(ScopeName)).Get();
 			}
 
+			CurrentScope->IncrementCount();
+
 			ScopeStack.Push(CurrentScope);
 		}
 
