PR #8848: RepGraph Bugfix: Prevent stale/garbage LastViewTarget actor being added to Replication List. (Contributed by [REDACTED_PERSON])

diff --git a/Engine/Plugins/Runtime/ReplicationGraph/Source/Private/ReplicationGraph.cpp b/Engine/Plugins/Runtime/ReplicationGraph/Source/Private/ReplicationGraph.cpp
--- a/Engine/Plugins/Runtime/ReplicationGraph/Source/Private/ReplicationGraph.cpp
+++ b/Engine/Plugins/Runtime/ReplicationGraph/Source/Private/ReplicationGraph.cpp
@@ -5950,7 +5950,7 @@ void UReplicationGraphNode_AlwaysRelevant_ForConnection::GatherActorListsForConn
 				// Zero out new actor cull distance
 				Params.ConnectionManager.ActorInfoMap.FindOrAdd(NewActor).SetCullDistanceSquared(0.f);
 			}
-			if (LastActor)
+			if (IsValid(LastActor))
 			{
 				// Reset previous actor culldistance
 				FConnectionReplicationActorInfo& ActorInfo = Params.ConnectionManager.ActorInfoMap.FindOrAdd(LastActor);
