Have UNetReplicationGraphConnection::PrepareForReplication check if the connection is in the correct world.
Handle emptying the map name if CurrentMapName == DestinationMapName when no transition map is specified, and remove LastCompletedSeamlessTravelCount == SeamlessTravelCount check from APlayerController::HasClientLoadedCurrentWorld to prevent issues when client finishes travel before the server.

diff --git a/Engine/Plugins/Runtime/ReplicationGraph/Source/Private/ReplicationGraph.cpp b/Engine/Plugins/Runtime/ReplicationGraph/Source/Private/ReplicationGraph.cpp
--- a/Engine/Plugins/Runtime/ReplicationGraph/Source/Private/ReplicationGraph.cpp
+++ b/Engine/Plugins/Runtime/ReplicationGraph/Source/Private/ReplicationGraph.cpp
@@ -2565,7 +2565,11 @@ void UNetReplicationGraphConnection::RemoveConnectionGraphNode(UReplicationGraph
 bool UNetReplicationGraphConnection::PrepareForReplication()
 {
 	NetConnection->ViewTarget = NetConnection->PlayerController ? NetConnection->PlayerController->GetViewTarget() : ToRawPtr(NetConnection->OwningActor);
-
+	
+	UWorld* CurrentWorld = GetWorld();
+	UPackage* CurrentWorldPackage = CurrentWorld ? CurrentWorld->GetPackage() : nullptr;
+	bool bConnectionHasCorrectWorld = CurrentWorldPackage ? NetConnection->GetClientWorldPackageName() == CurrentWorldPackage->GetFName() : true;
+	
 	// Set any children viewtargets
 	for (int32 i = 0; i < NetConnection->Children.Num(); ++i)
 	{
@@ -2573,7 +2577,7 @@ bool UNetReplicationGraphConnection::PrepareForReplication()
 		CurChild->ViewTarget = CurChild->PlayerController ? CurChild->PlayerController->GetViewTarget() : ToRawPtr(CurChild->OwningActor);
 	}
 
-	return (NetConnection->GetConnectionState() != USOCK_Closed) && (NetConnection->ViewTarget != nullptr);
+	return (NetConnection->GetConnectionState() != USOCK_Closed) && (NetConnection->ViewTarget != nullptr) && bConnectionHasCorrectWorld;
 }
 
 void UNetReplicationGraphConnection::NotifyAddDestructionInfo(FActorDestructionInfo* DestructInfo)
diff --git a/Engine/Source/Runtime/Engine/Private/PlayerController.cpp b/Engine/Source/Runtime/Engine/Private/PlayerController.cpp
--- a/Engine/Source/Runtime/Engine/Private/PlayerController.cpp
+++ b/Engine/Source/Runtime/Engine/Private/PlayerController.cpp
@@ -581,16 +581,7 @@ bool APlayerController::HasClientLoadedCurrentWorld()
 	if (Connection != NULL)
 	{
 		// NOTE: To prevent exploits, child connections must not use the parent connections ClientWorldPackageName value at all.
-		bool bInCorrectWorld = (Connection->GetClientWorldPackageName() == GetWorld()->GetOutermost()->GetFName());
-		if (SeamlessTravelCount > 0)
-		{
-			// In the case where seamless travel has occurred, make sure the client has actually completed the travel
-			return bInCorrectWorld && (LastCompletedSeamlessTravelCount == SeamlessTravelCount);
-		}
-		else
-		{
-			return bInCorrectWorld;
-		}
+		return (Connection->GetClientWorldPackageName() == GetWorld()->GetOutermost()->GetFName());
 	}
 	else
 	{
