PR #9167: Export FStreamingLevelActorListCollection to external modules (Contributed by [REDACTED_PERSON])

diff --git a/Engine/Plugins/Runtime/ReplicationGraph/Source/Public/ReplicationGraph.h b/Engine/Plugins/Runtime/ReplicationGraph/Source/Public/ReplicationGraph.h
--- a/Engine/Plugins/Runtime/ReplicationGraph/Source/Public/ReplicationGraph.h
+++ b/Engine/Plugins/Runtime/ReplicationGraph/Source/Public/ReplicationGraph.h
@@ -170,7 +170,7 @@ protected:
 // -----------------------------------
 
 
-struct FStreamingLevelActorListCollection
+struct REPLICATIONGRAPH_API FStreamingLevelActorListCollection
 {
 	void AddActor(const FNewReplicatedActorInfo& ActorInfo);
 	bool RemoveActor(const FNewReplicatedActorInfo& ActorInfo, bool bWarnIfNotFound, UReplicationGraphNode* Outer);
