ReplicationGraph
* Prevent actors outside the destruction radius but inside the current OutOfRange radius check from being added to the OutOfRange list.
* Instead keep actors in the PendingDestructionList until they get inside the send radius or outside the OutOfRange radius.

diff --git a/Engine/Plugins/Runtime/ReplicationGraph/Source/Private/ReplicationGraph.cpp b/Engine/Plugins/Runtime/ReplicationGraph/Source/Private/ReplicationGraph.cpp
--- a/Engine/Plugins/Runtime/ReplicationGraph/Source/Private/ReplicationGraph.cpp
+++ b/Engine/Plugins/Runtime/ReplicationGraph/Source/Private/ReplicationGraph.cpp
@@ -1041,6 +1041,8 @@ int32 UReplicationGraph::ServerReplicateActors(float DeltaSeconds)
 		ConnectionManager->GetClientVisibleLevelNames(AllVisibleLevelNames);
 		const FConnectionGatherActorListParameters Parameters(ConnectionViewers, *ConnectionManager, AllVisibleLevelNames, FrameNum, GatheredReplicationListsForConnection);
 
+		UNetReplicationGraphConnection::FRepGraphDestructionViewerInfoArray DestructionViewersInfo;
+
 		{
 			QUICK_SCOPE_CYCLE_COUNTER(NET_ReplicateActors_GatherForConnection);
 
@@ -1062,6 +1064,14 @@ int32 UReplicationGraph::ServerReplicateActors(float DeltaSeconds)
 				UE_LOG(LogReplicationGraph, Warning, TEXT("No Replication Lists were returned for connection"));
 				return 0;
 			}
+
+			for( const FNetViewer& NetViewer : ConnectionViewers )
+			{
+				FLastLocationGatherInfo* LastInfoForViewer = ConnectionManager->LastGatherLocations.FindByKey<UNetConnection*>(NetViewer.Connection);
+				check(LastInfoForViewer);
+
+				DestructionViewersInfo.Emplace(UNetReplicationGraphConnection::FRepGraphDestructionViewerInfo(NetViewer.ViewLocation, LastInfoForViewer->LastOutOfRangeLocationCheck));
+			}
 		}
 
 		// --------------------------------------------------------------------------------------------------------------
@@ -1141,7 +1151,7 @@ int32 UReplicationGraph::ServerReplicateActors(float DeltaSeconds)
 			// ------------------------------------------
 			{
 				QUICK_SCOPE_CYCLE_COUNTER(NET_ReplicateActors_ReplicateDestructionInfos);
-				ConnectionManager->ReplicateDestructionInfos(ConnectionViewers, DestructionSettings.DestructInfoMaxDistanceSquared);
+				ConnectionManager->ReplicateDestructionInfos(DestructionViewersInfo, DestructionSettings);
 			}
 
 			// ------------------------------------------
@@ -2756,30 +2766,38 @@ void UNetReplicationGraphConnection::NotifyClientVisibleLevelNamesAdd(FName Leve
 	}
 }
 
-int64 UNetReplicationGraphConnection::ReplicateDestructionInfos(const FNetViewerArray& Viewers, const float DestructInfoMaxDistanceSquared)
+int64 UNetReplicationGraphConnection::ReplicateDestructionInfos(const FRepGraphDestructionViewerInfoArray& DestructionViewersInfo, const FReplicationGraphDestructionSettings& DestructionSettings)
 {
-	QUICK_SCOPE_CYCLE_COUNTER(ReplicateDestructionInfos);
-
 	int64 NumBits = 0;
 #if WITH_SERVER_CODE
+
 	for (int32 idx=PendingDestructInfoList.Num()-1; idx >=0; --idx)
 	{
 		FCachedDestructInfo& Info = PendingDestructInfoList[idx];
 		FActorDestructionInfo* DestructInfo = Info.DestructionInfo;
 		bool bSendDestructionInfo = false;
+		bool bAddToOutOfRangeList = true;
 
 		if (!DestructInfo->bIgnoreDistanceCulling)
 		{
 			// Only send destruction info if the viewers are close enough to the destroyed actor
-			for (const FNetViewer& CurViewer : Viewers)
+			for (const FRepGraphDestructionViewerInfo& Viewer : DestructionViewersInfo)
 			{
-				const float DistSquared = FVector::DistSquared2D(Info.CachedPosition, CurViewer.ViewLocation);
+				const float DistSquared = FVector::DistSquared2D(Info.CachedPosition, Viewer.ViewerLocation);
 
-				if (DistSquared < DestructInfoMaxDistanceSquared)
+				if (DistSquared < DestructionSettings.DestructInfoMaxDistanceSquared)
 				{
 					bSendDestructionInfo = true;
 					break;
 				}
+
+				const float OutOfRangeDistSquared = FVector::DistSquared2D(Info.CachedPosition, Viewer.LastOutOfRangeLocationCheck);
+				
+				// Add the actor to the OutOfRangeList only if he is outside the range from the next check. If not keep him in the destruction list to be evaluated next frame.
+				if (OutOfRangeDistSquared <= DestructionSettings.MaxPendingListDistanceSquared)
+				{
+					bAddToOutOfRangeList = false;
+				}
 			}
 		}
 
@@ -2793,7 +2811,7 @@ int64 UNetReplicationGraphConnection::ReplicateDestructionInfos(const FNetViewer
 				TrackedDestructionInfoPtrs.Remove(DestructInfo);
 			}
 		}
-		else
+		else if (bAddToOutOfRangeList)
 		{
 			// Add the far actor to the out of range list so we don't evaluate it every frame
 			OutOfRangeDestroyedActors.Emplace(MoveTemp(Info));
@@ -2807,7 +2825,6 @@ int64 UNetReplicationGraphConnection::ReplicateDestructionInfos(const FNetViewer
 int64 UNetReplicationGraphConnection::ReplicateDormantDestructionInfos()
 {
 	CSV_SCOPED_TIMING_STAT_EXCLUSIVE(ReplicateDormantDestructionInfos);
-	QUICK_SCOPE_CYCLE_COUNTER(ReplicateDormantDestructionInfos);
 
 	int64 NumBits = 0;
 
@@ -2881,7 +2898,7 @@ void UNetReplicationGraphConnection::OnUpdateViewerLocation(FLastLocationGatherI
 
 			const float ActorDistSquared = FVector::DistSquared2D(CachedInfo.CachedPosition, Viewer.ViewLocation);
 				
-			if (ActorDistSquared < DestructionSettings.DestructInfoMaxDistanceSquared)
+			if (ActorDistSquared <= DestructionSettings.MaxPendingListDistanceSquared)
 			{
 				// Swap the info into the Pending List to get it replicated
 				PendingDestructInfoList.Emplace(MoveTemp(CachedInfo));
diff --git a/Engine/Plugins/Runtime/ReplicationGraph/Source/Public/ReplicationGraph.h b/Engine/Plugins/Runtime/ReplicationGraph/Source/Public/ReplicationGraph.h
--- a/Engine/Plugins/Runtime/ReplicationGraph/Source/Public/ReplicationGraph.h
+++ b/Engine/Plugins/Runtime/ReplicationGraph/Source/Public/ReplicationGraph.h
@@ -1301,6 +1301,21 @@ private:
 
 	friend UReplicationGraph;
 
+	/** Holds relevant data when parsing deleted actors that could be sent to a viewer */
+	struct FRepGraphDestructionViewerInfo
+	{
+		FRepGraphDestructionViewerInfo() = default;
+		FRepGraphDestructionViewerInfo(const FVector& InViewerLocation, const FVector& InOutOfRangeLocationCheck)
+			: ViewerLocation(InViewerLocation)
+			, LastOutOfRangeLocationCheck(InOutOfRangeLocationCheck)
+		{ }
+
+		FVector ViewerLocation;
+		FVector	LastOutOfRangeLocationCheck;
+	};
+
+	typedef TArray< FRepGraphDestructionViewerInfo, TInlineAllocator<REPGRAPH_VIEWERS_PER_CONNECTION> > FRepGraphDestructionViewerInfoArray;
+
 	// ----------------------------------------
 
 	/** Called right after this is created to associate with the owning Graph */
@@ -1317,7 +1332,7 @@ private:
 
 	bool PrepareForReplication();
 
-	int64 ReplicateDestructionInfos(const FNetViewerArray& InViewers, const float DestructInfoMaxDistanceSquared);
+	int64 ReplicateDestructionInfos(const FRepGraphDestructionViewerInfoArray& DestructionViewersInfo, const FReplicationGraphDestructionSettings& DestructionSettings);
 	
 	int64 ReplicateDormantDestructionInfos();
 
@@ -1453,8 +1468,10 @@ struct FReplicationGraphDestructionSettings
 	FReplicationGraphDestructionSettings(float InDestructInfoMaxDistanceSquared, float InOutOfRangeDistanceCheckThresholdSquared)
 		: DestructInfoMaxDistanceSquared(InDestructInfoMaxDistanceSquared)
 		, OutOfRangeDistanceCheckThresholdSquared(InOutOfRangeDistanceCheckThresholdSquared)
+		, MaxPendingListDistanceSquared(DestructInfoMaxDistanceSquared + OutOfRangeDistanceCheckThresholdSquared)
 	{ }
 
 	float DestructInfoMaxDistanceSquared;
 	float OutOfRangeDistanceCheckThresholdSquared;
+	float MaxPendingListDistanceSquared; 
 };
