Including some more comments and a warning around dependent actors in RepGraph

diff --git a/Engine/Plugins/Runtime/ReplicationGraph/Source/Private/ReplicationGraphTypes.cpp b/Engine/Plugins/Runtime/ReplicationGraph/Source/Private/ReplicationGraphTypes.cpp
--- a/Engine/Plugins/Runtime/ReplicationGraph/Source/Private/ReplicationGraphTypes.cpp
+++ b/Engine/Plugins/Runtime/ReplicationGraph/Source/Private/ReplicationGraphTypes.cpp
@@ -433,6 +433,12 @@ void TActorListAllocator<NumListsPerBlock, MaxNumPools>::LogDetails(int32 PoolSi
 
 void FGlobalActorReplicationInfoMap::AddDependentActor(AActor* Parent, AActor* Child, FGlobalActorReplicationInfoMap::EWarnFlag WarnFlag)
 {
+	/**
+	* We assume the relationship between the parent and the dependent to be one way: if the parent wants to replicate, then the dependent replicates.
+	* Dependent actors can still be routed into another node and replicated independently of their parent if desired, but if the dependent actor
+	* isn't being routed anywhere else, its parent must be routed somewhere in the RepGraph in order to ensure the dependent is replicated.
+	*/
+
 	const bool bIsParentValid = ensureMsgf(Parent && IsActorValidForReplication(Parent), TEXT("FGlobalActorReplicationInfoMap::AddDependentActor Invalid Parent! %s"),
 										   *GetPathNameSafe(Parent));
 
@@ -441,7 +447,8 @@ void FGlobalActorReplicationInfoMap::AddDependentActor(AActor* Parent, AActor* C
 
 	if (bIsParentValid && bIsChildValid)
 	{
-		const bool bDoWarnings = EnumHasAnyFlags(WarnFlag, EWarnFlag::WarnAlreadyDependant);
+		const bool bDoAlreadyDependantWarning = EnumHasAnyFlags(WarnFlag, EWarnFlag::WarnAlreadyDependant);
+		const bool bDoNotRegisteredWarning = EnumHasAnyFlags(WarnFlag, EWarnFlag::WarnParentNotRegistered);
 
 		bool bChildIsAlreadyDependant(false);
 		if (FGlobalActorReplicationInfo* ParentInfo = Find(Parent))
@@ -455,6 +462,10 @@ void FGlobalActorReplicationInfoMap::AddDependentActor(AActor* Parent, AActor* C
 				}
 			}
 		}
+		else if (bDoNotRegisteredWarning)
+		{
+			UE_LOG(LogReplicationGraph, Warning, TEXT("FGlobalActorReplicationInfoMap::AddDependentActor Parent %s not registered yet, cannot add child %s"), *GetNameSafe(Parent), *GetNameSafe(Child));
+		}
 
 		bool bChildHadParentAlready(false);
 		if (FGlobalActorReplicationInfo* ChildInfo = Find(Child))
@@ -469,7 +480,7 @@ void FGlobalActorReplicationInfoMap::AddDependentActor(AActor* Parent, AActor* C
 			}
 		}
 
-		if (bDoWarnings && (bChildIsAlreadyDependant || bChildHadParentAlready))
+		if (bDoAlreadyDependantWarning && (bChildIsAlreadyDependant || bChildHadParentAlready))
 		{
 			UE_LOG(LogReplicationGraph, Warning, TEXT("FGlobalActorReplicationInfoMap::AddDependentActor child %s already dependant of parent %s"), *GetNameSafe(Child), *GetNameSafe(Parent));
 		}
diff --git a/Engine/Plugins/Runtime/ReplicationGraph/Source/Public/ReplicationGraphTypes.h b/Engine/Plugins/Runtime/ReplicationGraph/Source/Public/ReplicationGraphTypes.h
--- a/Engine/Plugins/Runtime/ReplicationGraph/Source/Public/ReplicationGraphTypes.h
+++ b/Engine/Plugins/Runtime/ReplicationGraph/Source/Public/ReplicationGraphTypes.h
@@ -865,6 +865,8 @@ struct FGlobalActorReplicationInfoMap
 	{
 		None = 0,
 		WarnAlreadyDependant = 1 << 0,
+		WarnParentNotRegistered = 1 << 1,
+		AllWarnings = WarnAlreadyDependant | WarnParentNotRegistered
 	};
 
 	REPLICATIONGRAPH_API void AddDependentActor(AActor* Parent, AActor* Child, FGlobalActorReplicationInfoMap::EWarnFlag WarnFlag = FGlobalActorReplicationInfoMap::EWarnFlag::None);
