PR #9212: Fix grid cell node doubling child node actors (Contributed by [REDACTED_PERSON])

diff --git a/Engine/Plugins/Runtime/ReplicationGraph/Source/Private/ReplicationGraph.cpp b/Engine/Plugins/Runtime/ReplicationGraph/Source/Private/ReplicationGraph.cpp
--- a/Engine/Plugins/Runtime/ReplicationGraph/Source/Private/ReplicationGraph.cpp
+++ b/Engine/Plugins/Runtime/ReplicationGraph/Source/Private/ReplicationGraph.cpp
@@ -4704,14 +4704,6 @@ UReplicationGraphNode_DormancyNode* UReplicationGraphNode_GridCell::GetDormancyN
 void UReplicationGraphNode_GridCell::GetAllActorsInNode_Debugging(TArray<FActorRepListType>& OutArray) const
 {
 	Super::GetAllActorsInNode_Debugging(OutArray);
-	if (DynamicNode)
-	{
-		DynamicNode->GetAllActorsInNode_Debugging(OutArray);
-	}
-	if (DormancyNode)
-	{
-		DormancyNode->GetAllActorsInNode_Debugging(OutArray);
-	}
 }
 
 int32 CVar_RepGraph_DebugNextNewActor = 0;
@@ -6077,4 +6069,4 @@ FAutoConsoleCommandWithWorldAndArgs ChangeActorDiscoveryBudget(TEXT("Net.RepGrap
 	}
 }));
 
-#endif // WITH_SERVER_CODE
\ No newline at end of file
+#endif // WITH_SERVER_CODE
