Inclusive Terminology Push
Replaced "master" with "main" at ReplicationGraph.h/.cpp

diff --git a/Engine/Plugins/Runtime/ReplicationGraph/Source/Private/ReplicationGraph.cpp b/Engine/Plugins/Runtime/ReplicationGraph/Source/Private/ReplicationGraph.cpp
--- a/Engine/Plugins/Runtime/ReplicationGraph/Source/Private/ReplicationGraph.cpp
+++ b/Engine/Plugins/Runtime/ReplicationGraph/Source/Private/ReplicationGraph.cpp
@@ -4515,7 +4515,7 @@ UReplicationGraphNode_ConnectionDormancyNode* UReplicationGraphNode_DormancyNode
 	UReplicationGraphNode_ConnectionDormancyNode* ConnectionNode = CreateChildNode<UReplicationGraphNode_ConnectionDormancyNode>();
 	ConnectionNodes.Add(RepGraphConnection) = ConnectionNode;
 
-	// Copy our master lists to the connection node
+	// Copy our main lists to the connection node
 	ConnectionNode->DeepCopyActorListsFrom(this);
 
 	ConnectionNode->InitConnectionNode(RepGraphConnection, Params.ReplicationFrameNum);
diff --git a/Engine/Plugins/Runtime/ReplicationGraph/Source/Public/ReplicationGraph.h b/Engine/Plugins/Runtime/ReplicationGraph/Source/Public/ReplicationGraph.h
--- a/Engine/Plugins/Runtime/ReplicationGraph/Source/Public/ReplicationGraph.h
+++ b/Engine/Plugins/Runtime/ReplicationGraph/Source/Public/ReplicationGraph.h
@@ -498,7 +498,7 @@ private:
 };
 
 
-/** Stores per-connection copies of a master actor list. Skips and removes elements from per connection list that are fully dormant */
+/** Stores per-connection copies of a main actor list. Skips and removes elements from per connection list that are fully dormant */
 UCLASS()
 class REPLICATIONGRAPH_API UReplicationGraphNode_DormancyNode : public UReplicationGraphNode_ActorList
 {
