Disabling stat scope in test

diff --git a/Engine/Plugins/Runtime/ReplicationGraph/Source/Private/ReplicationGraph.cpp b/Engine/Plugins/Runtime/ReplicationGraph/Source/Private/ReplicationGraph.cpp
--- a/Engine/Plugins/Runtime/ReplicationGraph/Source/Private/ReplicationGraph.cpp
+++ b/Engine/Plugins/Runtime/ReplicationGraph/Source/Private/ReplicationGraph.cpp
@@ -4257,7 +4257,7 @@ bool ContainsReverse(const FActorRepListRefView& List, FActorRepListType Actor)
 
 void UReplicationGraphNode_ConnectionDormancyNode::NotifyActorDormancyFlush(FActorRepListType Actor)
 {
-	QUICK_SCOPE_CYCLE_COUNTER(ConnectionDormancyNode_NotifyActorDormancyFlush);
+	RG_QUICK_SCOPE_CYCLE_COUNTER(ConnectionDormancyNode_NotifyActorDormancyFlush);
 
 	FNewReplicatedActorInfo ActorInfo(Actor);
 
