Replication graph GridSpatialization2D node uses old world bounds values.
The grid is densely stored and risks very high memory usage at large coordinates.

diff --git a/Engine/Plugins/Runtime/ReplicationGraph/Source/Private/ReplicationGraph.cpp b/Engine/Plugins/Runtime/ReplicationGraph/Source/Private/ReplicationGraph.cpp
--- a/Engine/Plugins/Runtime/ReplicationGraph/Source/Private/ReplicationGraph.cpp
+++ b/Engine/Plugins/Runtime/ReplicationGraph/Source/Private/ReplicationGraph.cpp
@@ -68,6 +68,14 @@
 #include "PerfCountersModule.h"
 #endif
 
+namespace UE::Net::Private
+{
+	// Maintain pre-large world bounds for now, since the GridSpatialization2D node
+	// uses a densely-stored grid and risks very high memory usage at large coordinates.
+	constexpr double RepGraphWorldMax = UE_OLD_WORLD_MAX;
+	constexpr double RepGraphHalfWorldMax = UE_OLD_HALF_WORLD_MAX;
+}
+
 int32 CVar_RepGraph_Pause = 0;
 static FAutoConsoleVariableRef CVarRepGraphPause(TEXT("Net.RepGraph.Pause"), CVar_RepGraph_Pause,
 	TEXT("Pauses actor replication in the Replication Graph."), ECVF_Default);
@@ -4355,7 +4363,7 @@ bool UReplicationGraphNode_ConnectionDormancyNode::IsNodeObsolete(uint32 Current
 
 // --------------------------------------------------------------------------------------------------------------------------------------------
 
-FVector::FReal UReplicationGraphNode_DormancyNode::MaxZForConnection = WORLD_MAX;
+FVector::FReal UReplicationGraphNode_DormancyNode::MaxZForConnection = UE::Net::Private::RepGraphWorldMax;
 
 void UReplicationGraphNode_DormancyNode::CallFunctionOnValidConnectionNodes(FConnectionDormancyNodeFunction Function)
 {
@@ -4744,6 +4752,7 @@ static FAutoConsoleVariableRef CVarRepGraphDebugNextActor(TEXT("Net.RepGraph.Spa
 UReplicationGraphNode_GridSpatialization2D::UReplicationGraphNode_GridSpatialization2D()
 	: CellSize(0.f)
 	, SpatialBias(ForceInitToZero)
+	, ConnectionMaxZ(UE::Net::Private::RepGraphWorldMax)
 	, GridBounds(ForceInitToZero)
 {
 	bRequiresPrepareForReplicationCall = true;
@@ -5014,15 +5023,19 @@ void UReplicationGraphNode_GridSpatialization2D::GetGridNodesForActor(FActorRepL
 
 void UReplicationGraphNode_GridSpatialization2D::SetBiasAndGridBounds(const FBox& GridBox)
 {
+	using namespace UE::Net::Private;
+
 	const FVector2D BoxMin2D = FVector2D(GridBox.Min);
 	const FVector2D BoxMax2D = FVector2D(GridBox.Max);
 	
 	SpatialBias = BoxMin2D;
-	GridBounds = FBox( FVector(BoxMin2D, -HALF_WORLD_MAX), FVector(BoxMax2D, HALF_WORLD_MAX) );
+	GridBounds = FBox( FVector(BoxMin2D, -RepGraphHalfWorldMax), FVector(BoxMax2D, RepGraphHalfWorldMax) );
 }
 
 UReplicationGraphNode_GridSpatialization2D::FActorCellInfo UReplicationGraphNode_GridSpatialization2D::GetCellInfoForActor(FActorRepListType Actor, const FVector& Location3D, float CullDistance)
 {
+	using namespace UE::Net::Private;
+
 #if !(UE_BUILD_SHIPPING || UE_BUILD_TEST)
 	if (CullDistance <= 0.f)
 	{
@@ -5033,12 +5046,12 @@ UReplicationGraphNode_GridSpatialization2D::FActorCellInfo UReplicationGraphNode
 	FVector ClampedLocation = Location3D;
 
 	// Sanity check the actor's location. If it's garbage, we could end up with a gigantic allocation in GetGridNodesForActor as we adjust the grid.
-	if (Location3D.X < -HALF_WORLD_MAX || Location3D.X > HALF_WORLD_MAX ||
-		Location3D.Y < -HALF_WORLD_MAX || Location3D.Y > HALF_WORLD_MAX ||
-		Location3D.Z < -HALF_WORLD_MAX || Location3D.Z > HALF_WORLD_MAX)
+	if (Location3D.X < -RepGraphHalfWorldMax || Location3D.X > RepGraphHalfWorldMax ||
+		Location3D.Y < -RepGraphHalfWorldMax || Location3D.Y > RepGraphHalfWorldMax ||
+		Location3D.Z < -RepGraphHalfWorldMax || Location3D.Z > RepGraphHalfWorldMax)
 	{
 		UE_LOG(LogReplicationGraph, Warning, TEXT("GetCellInfoForActor: Actor %s is outside world bounds with a location of %s. Clamping grid location to world bounds."), *GetFullNameSafe(Actor), *Location3D.ToString());
-		ClampedLocation = Location3D.BoundToCube(HALF_WORLD_MAX);
+		ClampedLocation = Location3D.BoundToCube(RepGraphHalfWorldMax);
 	}
 
 	FActorCellInfo CellInfo;
@@ -5573,6 +5586,8 @@ struct FPlayerGridCellInformation
 
 void UReplicationGraphNode_GridSpatialization2D::GatherActorListsForConnection(const FConnectionGatherActorListParameters& Params)
 {
+	using namespace UE::Net::Private;
+
 #if WITH_SERVER_CODE
 	TArray<FLastLocationGatherInfo>& LastLocationArray = Params.ConnectionManager.LastGatherLocations;
 	TArray<FVector2D, FReplicationGraphConnectionsAllocator> UniqueCurrentLocations;
@@ -5597,7 +5612,7 @@ void UReplicationGraphNode_GridSpatialization2D::GatherActorListsForConnection(c
 		else
 		{
 			// Prevent extreme locations from causing the Grid to grow too large
-			ClampedViewLoc = ClampedViewLoc.BoundToCube(HALF_WORLD_MAX);
+			ClampedViewLoc = ClampedViewLoc.BoundToCube(RepGraphHalfWorldMax);
 		}
 
 		// Find out what bucket the view is in
@@ -5642,7 +5657,7 @@ void UReplicationGraphNode_GridSpatialization2D::GatherActorListsForConnection(c
 		else
 		{
 			// Prevent extreme locations from causing the Grid to grow too large
-			LastLocationForConnection = LastLocationForConnection.BoundToCube(HALF_WORLD_MAX);
+			LastLocationForConnection = LastLocationForConnection.BoundToCube(RepGraphHalfWorldMax);
 		}
 
 		// Try to determine the previous location of the user.
diff --git a/Engine/Plugins/Runtime/ReplicationGraph/Source/Public/ReplicationGraph.h b/Engine/Plugins/Runtime/ReplicationGraph/Source/Public/ReplicationGraph.h
--- a/Engine/Plugins/Runtime/ReplicationGraph/Source/Public/ReplicationGraph.h
+++ b/Engine/Plugins/Runtime/ReplicationGraph/Source/Public/ReplicationGraph.h
@@ -619,7 +619,7 @@ public:
 	
 	float		CellSize;
 	FVector2D	SpatialBias;
-	float		ConnectionMaxZ = WORLD_MAX; // Connection locations have to be <= to this to pull from the grid
+	double		ConnectionMaxZ; // Connection locations have to be <= to this to pull from the grid
 
 	/** When the GridBounds is set we limit the creation of cells to be exclusively inside the passed region.
 	    Viewers who gather nodes outside this region will be clamped to the closest cell inside the box.
