Handle an actor with an invalid location (NaN) in GetCellInfoForActor

diff --git a/Engine/Plugins/Runtime/ReplicationGraph/Source/Private/ReplicationGraph.cpp b/Engine/Plugins/Runtime/ReplicationGraph/Source/Private/ReplicationGraph.cpp
--- a/Engine/Plugins/Runtime/ReplicationGraph/Source/Private/ReplicationGraph.cpp
+++ b/Engine/Plugins/Runtime/ReplicationGraph/Source/Private/ReplicationGraph.cpp
@@ -5054,6 +5054,11 @@ UReplicationGraphNode_GridSpatialization2D::FActorCellInfo UReplicationGraphNode
 		ClampedLocation = Location3D.BoundToCube(RepGraphHalfWorldMax);
 	}
 
+	if (!ensureMsgf(!ClampedLocation.ContainsNaN(), TEXT("GetCellInfoForActor: Actor %s has an invalid location of %s, defaulting to the origin."), *GetFullNameSafe(Actor), *ClampedLocation.ToString()))
+	{
+		ClampedLocation = FVector::ZeroVector;
+	}
+
 	FActorCellInfo CellInfo;
 	const auto LocationBiasX = (ClampedLocation.X - SpatialBias.X);
 	const auto LocationBiasY = (ClampedLocation.Y - SpatialBias.Y);
