Add a CSV stat to track CPU time impact of multiple replications of the same actor to the same connection in the same frame. Possible in RepGraph when an RPC causes a new channel to open.
Change FReplicationGraphCSVTracker::PostReplicateActor to take enum flags.

diff --git a/Engine/Plugins/Runtime/ReplicationGraph/Source/Private/ReplicationGraph.cpp b/Engine/Plugins/Runtime/ReplicationGraph/Source/Private/ReplicationGraph.cpp
--- a/Engine/Plugins/Runtime/ReplicationGraph/Source/Private/ReplicationGraph.cpp
+++ b/Engine/Plugins/Runtime/ReplicationGraph/Source/Private/ReplicationGraph.cpp
@@ -161,6 +161,7 @@ CSV_DEFINE_CATEGORY(ReplicationGraphVisibleLevels, WITH_SERVER_CODE);
 CSV_DEFINE_CATEGORY(ReplicationGraphForcedUpdates, WITH_SERVER_CODE);
 CSV_DEFINE_CATEGORY(ReplicationGraphCleanMS, WITH_SERVER_CODE);
 CSV_DEFINE_CATEGORY(ReplicationGraphCleanNumReps, WITH_SERVER_CODE);
+CSV_DEFINE_CATEGORY(ReplicationGraphRedundantMS, WITH_SERVER_CODE);
 
 CSV_DEFINE_CATEGORY(ReplicationGraph, WITH_SERVER_CODE);
 
@@ -1832,6 +1833,8 @@ int64 UReplicationGraph::ReplicateSingleActor(AActor* Actor, FConnectionReplicat
 		return 0;
 	}
 
+	FReplicationGraphCSVTracker::EActorFlags CSVFlags = FReplicationGraphCSVTracker::EActorFlags::None;
+
 	if (LIKELY(ActorInfo.Channel))
 	{
 		if (UNLIKELY(ActorInfo.Channel->Closing))
@@ -1843,6 +1846,17 @@ int64 UReplicationGraph::ReplicateSingleActor(AActor* Actor, FConnectionReplicat
 		{
 			return 0;
 		}
+
+		// If this actor has passed the repgraph last-rep-frame check for normal replication,
+		// repgraph thinks the actor has not been replicated to this connection yet this frame.
+		// However, if the actor was replicated due to an RPC opening a new channel, repgraph is unaware
+		// and will attempt to replicate it here.
+#if REPGRAPH_CSV_TRACKER
+		if (ActorInfo.Channel->LastUpdateTime == NetDriver->GetElapsedTime())
+		{
+			CSVFlags |= FReplicationGraphCSVTracker::EActorFlags::AlreadyReplicatedThisFrame;
+		}
+#endif
 	}
 
 	FrameReplicationStats.NumReplicatedActors++;
@@ -1935,7 +1949,12 @@ int64 UReplicationGraph::ReplicateSingleActor(AActor* Actor, FConnectionReplicat
 	const bool bIsTrafficActorDiscovery = ActorDiscoveryMaxBitsPerFrame > 0 && (ActorInfo.Channel && ActorInfo.Channel->SpawnAcked == false);
 	const bool bIsActorDiscoveryBudgetFull = bIsTrafficActorDiscovery && (ConnectionManager.QueuedBitsForActorDiscovery >= ActorDiscoveryMaxBitsPerFrame);
 
-	CSVTracker.PostReplicateActor(ActorClass, DeltaReplicateActorTimeSeconds, BitsWritten, bIsTrafficActorDiscovery && !bIsActorDiscoveryBudgetFull);
+	if (bIsTrafficActorDiscovery && !bIsActorDiscoveryBudgetFull)
+	{
+		CSVFlags |= FReplicationGraphCSVTracker::EActorFlags::IsInDiscovery;
+	}
+
+	CSVTracker.PostReplicateActor(ActorClass, DeltaReplicateActorTimeSeconds, BitsWritten, CSVFlags);
 
 	// ----------------------------
 	//	Dependent actors
diff --git a/Engine/Plugins/Runtime/ReplicationGraph/Source/Public/ReplicationGraphTypes.h b/Engine/Plugins/Runtime/ReplicationGraph/Source/Public/ReplicationGraphTypes.h
--- a/Engine/Plugins/Runtime/ReplicationGraph/Source/Public/ReplicationGraphTypes.h
+++ b/Engine/Plugins/Runtime/ReplicationGraph/Source/Public/ReplicationGraphTypes.h
@@ -8,6 +8,7 @@
 #include "Engine/LocalPlayer.h"
 #include "GameFramework/Actor.h"
 #include "Misc/ConfigCacheIni.h"
+#include "Misc/EnumClassFlags.h"
 #include "Net/DataBunch.h"
 #include "UObject/ObjectKey.h"
 #include "UObject/UObjectHash.h"
@@ -1475,6 +1476,7 @@ CSV_DECLARE_CATEGORY_EXTERN(ReplicationGraphVisibleLevels);
 CSV_DECLARE_CATEGORY_EXTERN(ReplicationGraphForcedUpdates);
 CSV_DECLARE_CATEGORY_EXTERN(ReplicationGraphCleanMS);
 CSV_DECLARE_CATEGORY_EXTERN(ReplicationGraphCleanNumReps);
+CSV_DECLARE_CATEGORY_EXTERN(ReplicationGraphRedundantMS);
 
 #ifndef REPGRAPH_CSV_TRACKER
 #define REPGRAPH_CSV_TRACKER (CSV_PROFILER && WITH_SERVER_CODE)
@@ -1483,6 +1485,15 @@ CSV_DECLARE_CATEGORY_EXTERN(ReplicationGraphCleanNumReps);
 /** Helper struct for tracking finer grained ReplicationGraph stats through the CSV profiler. Intention is that it is setup/configured in the UReplicationGraph subclasses */
 struct FReplicationGraphCSVTracker
 {
+	enum class EActorFlags : uint8
+	{
+		None = 0,
+		IsInDiscovery = 1 << 0,
+		AlreadyReplicatedThisFrame = 1 << 1
+	};
+
+	FRIEND_ENUM_CLASS_FLAGS(EActorFlags)
+
 	FReplicationGraphCSVTracker()
 		: EverythingElse(TEXT("Other"))
 		, EverythingElse_FastPath(TEXT("OtherFastPath"))
@@ -1595,7 +1606,15 @@ struct FReplicationGraphCSVTracker
 #endif //REPGRAPH_CSV_TRACKER
 	}
 
+	UE_DEPRECATED(5.1, "Use the overload of PostReplicateActor that takes EActorFlags.")
 	void PostReplicateActor(UClass* ActorClass, const double Time, const int64 Bits, const bool bIsActorDiscovery)
+	{
+		const EActorFlags Flags = bIsActorDiscovery ? EActorFlags::IsInDiscovery : EActorFlags::None;
+
+		PostReplicateActor(ActorClass, Time, Bits, Flags);
+	}
+
+	void PostReplicateActor(UClass* ActorClass, const double Time, const int64 Bits, const EActorFlags Flags)
 	{
 #if REPGRAPH_CSV_TRACKER
 		if (!bIsCapturing)
@@ -1620,11 +1639,16 @@ struct FReplicationGraphCSVTracker
 		}
 
 		// When opening actor channels keep all traffic in a separate bucket
-		if (bIsActorDiscovery)
+		if (EnumHasAnyFlags(Flags, EActorFlags::IsInDiscovery))
 		{
 			ActorDiscovery.BitsAccumulated += Bits;
 			ActorDiscovery.CPUTimeAccumulated += Time;
 
+			if (EnumHasAnyFlags(Flags, EActorFlags::AlreadyReplicatedThisFrame))
+			{
+				ActorDiscovery.RedundantRepCPUTimeAccumulated += Time;
+			}
+
 			// But keep the number of replicated classes unique
 			TrackedData->NumReplications++;
 		}
@@ -1639,6 +1663,11 @@ struct FReplicationGraphCSVTracker
 				TrackedData->CleanCPUTimeAccumulated += Time;
 				TrackedData->CleanNumReplications++;
 			}
+
+			if (EnumHasAnyFlags(Flags, EActorFlags::AlreadyReplicatedThisFrame))
+			{
+				TrackedData->RedundantRepCPUTimeAccumulated += Time;
+			}
 		}
 #endif	
 	}
@@ -1811,6 +1840,7 @@ private:
 
 		double CPUTimeAccumulated = 0.0;
 		double CleanCPUTimeAccumulated = 0.0;
+		double RedundantRepCPUTimeAccumulated = 0.0;
 		int64 BitsAccumulated = 0;
 		int32 ChannelsOpened = 0;
 		int32 NumReplications = 0;
@@ -1823,6 +1853,7 @@ private:
 		{
 			CPUTimeAccumulated = 0.0;
 			CleanCPUTimeAccumulated = 0.0;
+			RedundantRepCPUTimeAccumulated = 0.0;
 			BitsAccumulated = 0;
 			ChannelsOpened = 0;
 			NumReplications = 0;
@@ -1863,12 +1894,15 @@ private:
 		Profiler->RecordCustomStat(Data.StatName, CSV_CATEGORY_INDEX(ReplicationGraphForcedUpdates), static_cast<float>(Data.ForcedUpdates), ECsvCustomStatOp::Set);
 		Profiler->RecordCustomStat(Data.StatName, CSV_CATEGORY_INDEX(ReplicationGraphCleanMS), static_cast<float>(Data.CleanCPUTimeAccumulated) * 1000.f, ECsvCustomStatOp::Set);
 		Profiler->RecordCustomStat(Data.StatName, CSV_CATEGORY_INDEX(ReplicationGraphCleanNumReps), static_cast<float>(Data.CleanNumReplications), ECsvCustomStatOp::Set);
+		Profiler->RecordCustomStat(Data.StatName, CSV_CATEGORY_INDEX(ReplicationGraphRedundantMS), static_cast<float>(Data.RedundantRepCPUTimeAccumulated) * 1000.f, ECsvCustomStatOp::Set);
 
 		Data.Reset();
 	}
 #endif
 };
 
+ENUM_CLASS_FLAGS(FReplicationGraphCSVTracker::EActorFlags);
+
 // Debug Actor/Connection pair that can be set by code for further narrowing down breakpoints/logging
 struct FActorConnectionPair
 {
