Fix level visibility check for dormant actor cleanup feature. This allows dormant actors placed in a persistent level to be destroyed on the client and reclaim their memory.

#tests PIE

diff --git a/Engine/Plugins/Runtime/ReplicationGraph/Source/Private/ReplicationGraph.cpp b/Engine/Plugins/Runtime/ReplicationGraph/Source/Private/ReplicationGraph.cpp
--- a/Engine/Plugins/Runtime/ReplicationGraph/Source/Private/ReplicationGraph.cpp
+++ b/Engine/Plugins/Runtime/ReplicationGraph/Source/Private/ReplicationGraph.cpp
@@ -2687,15 +2687,14 @@ void UNetReplicationGraphConnection::NotifyAddDormantDestructionInfo(AActor* Act
 	if (Actor && NetConnection && NetConnection->Driver && NetConnection->Driver->GuidCache)
 	{
 		ULevel* Level = Actor->GetLevel();
-		FName StreamingLevelName = NAME_None;
 
-		if (Level && Level->IsPersistentLevel())
+		if (Level && !Level->IsPersistentLevel())
 		{
-			StreamingLevelName = Level->GetOutermost()->GetFName();
+			FName StreamingLevelName = Level->GetOutermost()->GetFName();
 
 			if (NetConnection->ClientVisibleLevelNames.Contains(StreamingLevelName) == false)
 			{
-				UE_LOG(LogReplicationGraph, Verbose, TEXT("NotifyAddDormantDestructionInfo skipping actor [%s] because streaming level is no longer visible."), *GetNameSafe(Actor));
+				UE_LOG(LogReplicationGraph, Verbose, TEXT("NotifyAddDormantDestructionInfo skipping actor [%s] because streaming level [%s] is no longer visible."), *GetNameSafe(Actor), *StreamingLevelName.ToString());
 				return;
 			}
 		}
