Replace some rep graph node strong refs with weak pointers

diff --git a/Engine/Plugins/Runtime/ReplicationGraph/Source/Private/ReplicationGraph.cpp b/Engine/Plugins/Runtime/ReplicationGraph/Source/Private/ReplicationGraph.cpp
--- a/Engine/Plugins/Runtime/ReplicationGraph/Source/Private/ReplicationGraph.cpp
+++ b/Engine/Plugins/Runtime/ReplicationGraph/Source/Private/ReplicationGraph.cpp
@@ -5196,6 +5196,7 @@ void UReplicationGraphNode_GridSpatialization2D::PrepareForReplication()
 	RG_QUICK_SCOPE_CYCLE_COUNTER(UReplicationGraphNode_GridSpatialization2D_PrepareForReplication);
 
 	FGlobalActorReplicationInfoMap* GlobalRepMap = GraphGlobals.IsValid() ? GraphGlobals->GlobalActorReplicationInfoMap : nullptr;
+	check(GlobalRepMap);
 
 #if !(UE_BUILD_SHIPPING || UE_BUILD_TEST)
 	if (CVar_RepGraph_Spatial_BiasCreep != 0.f)
@@ -6016,38 +6017,43 @@ void UReplicationGraphNode_TearOff_ForConnection::NotifyTearOffActor(AActor* Act
 
 // -------------------------------------------------------
 
-void UReplicationGraphNode_AlwaysRelevant_ForConnection::GatherActorListsForConnection(const FConnectionGatherActorListParameters& Params)
+void UReplicationGraphNode_AlwaysRelevant_ForConnection::AddCachedRelevantActor(const FConnectionGatherActorListParameters& Params, AActor* NewActor, TWeakObjectPtr<AActor>& LastActor)
 {
-	// Call super to add any actors that were explicitly given to use via NotifyAddNetworkActor
-	Super::GatherActorListsForConnection(Params);
+	UpdateCachedRelevantActor(Params, NewActor, LastActor);
 
-#if WITH_SERVER_CODE
-	auto UpdateActor = [&](AActor* NewActor, AActor*& LastActor)
+	if (NewActor && !ReplicationActorList.Contains(NewActor))
 	{
-		if (NewActor != LastActor)
-		{
-			if (NewActor)
-			{
-				// Zero out new actor cull distance
-				Params.ConnectionManager.ActorInfoMap.FindOrAdd(NewActor).SetCullDistanceSquared(0.f);
-			}
-			if (IsValid(LastActor))
-			{
-				// Reset previous actor culldistance
-				FConnectionReplicationActorInfo& ActorInfo = Params.ConnectionManager.ActorInfoMap.FindOrAdd(LastActor);
-				ActorInfo.SetCullDistanceSquared(GraphGlobals->GlobalActorReplicationInfoMap->Get(LastActor).Settings.GetCullDistanceSquared());
-			}
-
-			LastActor = NewActor;
+		ReplicationActorList.Add(NewActor);
+	}
+}
 
+void UReplicationGraphNode_AlwaysRelevant_ForConnection::UpdateCachedRelevantActor(const FConnectionGatherActorListParameters& Params, AActor* NewActor, TWeakObjectPtr<AActor>& LastActor)
+{
+	if (NewActor != LastActor)
+	{
+		if (NewActor)
+		{
+			// Zero out new actor cull distance
+			Params.ConnectionManager.ActorInfoMap.FindOrAdd(NewActor).SetCullDistanceSquared(0.f);
 		}
 
-		if (NewActor && !ReplicationActorList.Contains(NewActor))
+		if (AActor* PrevActor = LastActor.Get())
 		{
-			ReplicationActorList.Add(NewActor);
+			// Reset previous actor culldistance
+			FConnectionReplicationActorInfo& ActorInfo = Params.ConnectionManager.ActorInfoMap.FindOrAdd(PrevActor);
+			ActorInfo.SetCullDistanceSquared(GraphGlobals->GlobalActorReplicationInfoMap->Get(PrevActor).Settings.GetCullDistanceSquared());
 		}
-	};
 
+		LastActor = NewActor;
+	}
+}
+
+void UReplicationGraphNode_AlwaysRelevant_ForConnection::GatherActorListsForConnection(const FConnectionGatherActorListParameters& Params)
+{
+	// Call super to add any actors that were explicitly given to use via NotifyAddNetworkActor
+	Super::GatherActorListsForConnection(Params);
+
+#if WITH_SERVER_CODE
 	// Reset and rebuild another list that will contains our current viewer/viewtarget
 	ReplicationActorList.Reset();
 
@@ -6064,27 +6070,14 @@ void UReplicationGraphNode_AlwaysRelevant_ForConnection::GatherActorListsForConn
 			continue;
 		}
 
-		FAlwaysRelevantActorInfo* LastData = PastRelevantActors.FindByKey<UNetConnection*>(CurViewer.Connection);
-
-		// We've not seen this actor before, go ahead and add them.
-		if (LastData == nullptr)
-		{
-			FAlwaysRelevantActorInfo NewActorInfo;
-			NewActorInfo.Connection = CurViewer.Connection;
-			LastData = &(PastRelevantActors[PastRelevantActors.Add(NewActorInfo)]);
-		}
-
-		check(LastData != nullptr);
+		FCachedAlwaysRelevantActorInfo& LastData = PastRelevantActorMap.FindOrAdd(CurViewer.Connection);
 
-		UpdateActor(CurViewer.InViewer, static_cast<AActor*&>(LastData->LastViewer));
-		UpdateActor(CurViewer.ViewTarget, static_cast<AActor*&>(LastData->LastViewTarget));
+		AddCachedRelevantActor(Params, CurViewer.InViewer, LastData.LastViewer);
+		AddCachedRelevantActor(Params, CurViewer.ViewTarget, LastData.LastViewTarget);
 	}
 
-	// Remove excess
-	PastRelevantActors.RemoveAll([&](FAlwaysRelevantActorInfo& RelActorInfo){
-		return RelActorInfo.Connection == nullptr;
-	});
-	
+	CleanupCachedRelevantActors(PastRelevantActorMap);
+
 	if (ReplicationActorList.Num() > 0)
 	{
 		Params.OutGatheredReplicationLists.AddReplicationActorList(ReplicationActorList);
diff --git a/Engine/Plugins/Runtime/ReplicationGraph/Source/Public/ReplicationGraph.h b/Engine/Plugins/Runtime/ReplicationGraph/Source/Public/ReplicationGraph.h
--- a/Engine/Plugins/Runtime/ReplicationGraph/Source/Public/ReplicationGraph.h
+++ b/Engine/Plugins/Runtime/ReplicationGraph/Source/Public/ReplicationGraph.h
@@ -799,7 +799,7 @@ protected:
 // -----------------------------------
 
 USTRUCT()
-struct FAlwaysRelevantActorInfo
+struct UE_DEPRECATED(5.2, "No longer used") FAlwaysRelevantActorInfo
 {
 	GENERATED_BODY()
 
@@ -818,6 +818,12 @@ struct FAlwaysRelevantActorInfo
 	}
 };
 
+struct FCachedAlwaysRelevantActorInfo
+{
+	TWeakObjectPtr<AActor> LastViewer;
+	TWeakObjectPtr<AActor> LastViewTarget;
+};
+
 /** Adds actors that are always relevant for a connection. This engine version just adds the PlayerController and ViewTarget (usually the pawn) */
 UCLASS()
 class REPLICATIONGRAPH_API UReplicationGraphNode_AlwaysRelevant_ForConnection : public UReplicationGraphNode_ActorList
@@ -833,14 +839,32 @@ public:
 	/** Rebuilt-every-frame list based on UNetConnection state */
 	FActorRepListRefView ReplicationActorList;
 
+	PRAGMA_DISABLE_DEPRECATION_WARNINGS
 	/** List of previously (or currently if nothing changed last tick) focused actor data per connection */
+	UE_DEPRECATED(5.2, "Please use PastRelevantActorMap instead")
 	UPROPERTY()
 	TArray<FAlwaysRelevantActorInfo> PastRelevantActors;
+	PRAGMA_ENABLE_DEPRECATION_WARNINGS
 
-protected:
+	template<typename KeyType, typename ValueType>
+	static void CleanupCachedRelevantActors(TMap<TObjectKey<KeyType>, ValueType>& ActorMap)
+	{
+		for (auto MapIt = ActorMap.CreateIterator(); MapIt; ++MapIt)
+		{
+			if (MapIt.Key().ResolveObjectPtr() == nullptr)
+			{
+				MapIt.RemoveCurrent();
+			}
+		}
+	}
 
+protected:
 	virtual void OnCollectActorRepListStats(struct FActorRepListStatCollector& StatsCollector) const override;
 
+	void AddCachedRelevantActor(const FConnectionGatherActorListParameters& Params, AActor* NewActor, TWeakObjectPtr<AActor>& LastActor);
+	void UpdateCachedRelevantActor(const FConnectionGatherActorListParameters& Params, AActor* NewActor, TWeakObjectPtr<AActor>& LastActor);
+
+	TMap<TObjectKey<UNetConnection>, FCachedAlwaysRelevantActorInfo> PastRelevantActorMap;
 };
 
 // -----------------------------------
