RepGraph
* Deleted deprecated code from before 5.1

diff --git a/Engine/Plugins/Runtime/ReplicationGraph/Source/Private/ReplicationGraph.cpp b/Engine/Plugins/Runtime/ReplicationGraph/Source/Private/ReplicationGraph.cpp
--- a/Engine/Plugins/Runtime/ReplicationGraph/Source/Private/ReplicationGraph.cpp
+++ b/Engine/Plugins/Runtime/ReplicationGraph/Source/Private/ReplicationGraph.cpp
@@ -478,9 +478,6 @@ UNetReplicationGraphConnection* UReplicationGraph::CreateClientConnectionManager
 	// Give it an ID
 	const int32 NewConnectionNum = Connections.Num() + PendingConnections.Num();
 	NewConnectionManager->ConnectionOrderNum = NewConnectionNum;
-	PRAGMA_DISABLE_DEPRECATION_WARNINGS
-	NewConnectionManager->ConnectionId = NewConnectionNum;
-	PRAGMA_ENABLE_DEPRECATION_WARNINGS
 
 	// Initialize it with us
 	NewConnectionManager->InitForGraph(this);
@@ -506,9 +503,6 @@ UNetReplicationGraphConnection* UReplicationGraph::FixGraphConnectionList(TArray
 			// Fix the ConnectionOrderNum
 			const int32 NewConnectionNum = ConnectionNum++;
 			CurrentGraphConnection->ConnectionOrderNum = NewConnectionNum;
-			PRAGMA_DISABLE_DEPRECATION_WARNINGS
-			CurrentGraphConnection->ConnectionId = NewConnectionNum;
-			PRAGMA_ENABLE_DEPRECATION_WARNINGS
 		}
 		else
 		{
@@ -3891,19 +3885,6 @@ FORCEINLINE uint32 CalcDynamicReplicationPeriod(const float FinalPCT, const uint
 
 static TArray<FColor> DynamicSpatialFrequencyDebugColorArray = { FColor::Red, FColor::Green, FColor::Blue, FColor::Cyan, FColor::Orange, FColor::Purple };
 
-void UReplicationGraphNode_DynamicSpatialFrequency::CalcFrequencyForActor(AActor* Actor, UReplicationGraph* RepGraph, UNetConnection* NetConnection, FGlobalActorReplicationInfo& GlobalInfo, FConnectionReplicationActorInfo& ConnectionInfo, FSettings& MySettings, const FNetViewerArray& Viewers, const uint32 FrameNum, int32 ExistingItemIndex)
-{
-	for (UNetReplicationGraphConnection* ConnectionManager : RepGraph->Connections)
-	{
-		if (ConnectionManager->NetConnection == NetConnection)
-		{
-			CalcFrequencyForActor(Actor, RepGraph, NetConnection, GlobalInfo, ConnectionManager->ActorInfoMap, MySettings, Viewers, FrameNum, ExistingItemIndex);
-			return;
-		}
-	}
-}
-
-
 void UReplicationGraphNode_DynamicSpatialFrequency::CalcFrequencyForActor(AActor* Actor, UReplicationGraph* RepGraph, UNetConnection* NetConnection, FGlobalActorReplicationInfo& GlobalInfo, FPerConnectionActorInfoMap& ConnectionMap, FSettings& MySettings, const FNetViewerArray& Viewers, const uint32 FrameNum, int32 ExistingItemIndex)
 {
 #if WITH_SERVER_CODE
diff --git a/Engine/Plugins/Runtime/ReplicationGraph/Source/Public/ReplicationGraph.h b/Engine/Plugins/Runtime/ReplicationGraph/Source/Public/ReplicationGraph.h
--- a/Engine/Plugins/Runtime/ReplicationGraph/Source/Public/ReplicationGraph.h
+++ b/Engine/Plugins/Runtime/ReplicationGraph/Source/Public/ReplicationGraph.h
@@ -434,9 +434,6 @@ protected:
 	virtual void GatherActors(const FActorRepListRefView& RepList, FGlobalActorReplicationInfoMap& GlobalMap, FPerConnectionActorInfoMap& ConnectionMap, const FConnectionGatherActorListParameters& Params, UNetConnection* NetConnection);
 	virtual void GatherActors_DistanceOnly(const FActorRepListRefView& RepList, FGlobalActorReplicationInfoMap& GlobalMap, FPerConnectionActorInfoMap& ConnectionMap, const FConnectionGatherActorListParameters& Params);
 
-	UE_DEPRECATED(4.27, "Please use version of CalcFrequencyForActor that takes FPerConnectionActorInfoMap instead of FConnectionReplicationActorInfo.")
-	void CalcFrequencyForActor(AActor* Actor, UReplicationGraph* RepGraph, UNetConnection* NetConnection, FGlobalActorReplicationInfo& GlobalInfo, FConnectionReplicationActorInfo& ConnectionInfo, FSettings& MySettings, const FNetViewerArray& Viewers, const uint32 FrameNum, int32 ExistingItemIndex);
-	
 	void CalcFrequencyForActor(AActor* Actor, UReplicationGraph* RepGraph, UNetConnection* NetConnection, FGlobalActorReplicationInfo& GlobalInfo, FPerConnectionActorInfoMap& ConnectionMap, FSettings& MySettings, const FNetViewerArray& Viewers, const uint32 FrameNum, int32 ExistingItemIndex);
 };
 
@@ -635,9 +632,6 @@ public:
 	// Marks a class as preventing spatial rebuilds when an instance leaves the grid
 	void AddToClassRebuildDenyList(UClass* Class) { ClassRebuildDenyList.Set(Class, true); }
 
-	UE_DEPRECATED(5.0, "Use AddToClassRebuildDenyList instead")
-	void AddSpatialRebuildBlacklistClass(UClass* Class) { AddToClassRebuildDenyList(Class); }
-
 #if !(UE_BUILD_SHIPPING || UE_BUILD_TEST)
 	TArray<FString> DebugActorNames;
 #endif
@@ -1269,10 +1263,6 @@ public:
 	/** Index of the connection in the global list. Will be reassigned when any client disconnects so it is a key that can be referenced only during a single tick */
 	int32 ConnectionOrderNum;
 
-	// ID that is assigned by the replication graph. Will be reassigned/compacted as clients disconnect. Useful for spacing out connection operations. E.g., not stable but always compact.
-	UE_DEPRECATED(4.26, "This variable was renamed to ConnectionOrderNum to better reflect that it is not persistent and should not be considered an ID.")
-	int32 ConnectionId; 
-
 	UPROPERTY()
 	TArray<FLastLocationGatherInfo> LastGatherLocations;
 
diff --git a/Engine/Plugins/Runtime/ReplicationGraph/Source/Public/ReplicationGraphTypes.h b/Engine/Plugins/Runtime/ReplicationGraph/Source/Public/ReplicationGraphTypes.h
--- a/Engine/Plugins/Runtime/ReplicationGraph/Source/Public/ReplicationGraphTypes.h
+++ b/Engine/Plugins/Runtime/ReplicationGraph/Source/Public/ReplicationGraphTypes.h
@@ -241,16 +241,6 @@ struct REPLICATIONGRAPH_API FActorRepListRefView
 	{
 	}
 
-	UE_DEPRECATED(4.27, "This constructor is deprecated since the class does not hold FActorRepList's anymore. ")
-	FActorRepListRefView(FActorRepList& InRepList) 
-	{ 
-		RepList.Reserve(InRepList.Num);
-		for (int32 i = 0; i < InRepList.Num; ++i)
-		{
-			RepList.Add(InRepList.Data[i]);
-		}
-	}
-
 	/** Empties the array but does not deallocate the internal memory. Will be resized if the specified max size is bigger than the current max. */
 	void Reset(int32 ExpectedMaxSize=0)
 	{
@@ -263,12 +253,6 @@ struct REPLICATIONGRAPH_API FActorRepListRefView
 		RepList.Reserve(Size);
 	}
 
-	UE_DEPRECATED(4.27, "PrepareForWrite is not needed before calling operations on the RepList anymore. Use Reserve or Reset if you want to preallocate the array to a specific size")
-	void PrepareForWrite(bool bResetContent=false)
-	{
-		// empty
-	}
-
 	bool ConditionalAdd(const FActorRepListType& NewElement)
 	{
 		if (IsActorValidForReplicationGather(NewElement))
@@ -284,12 +268,6 @@ struct REPLICATIONGRAPH_API FActorRepListRefView
 		RepList.Add(NewElement);
 	}
 
-	UE_DEPRECATED(4.27, "Remove has been deprecated in favor of RemoveSlow/RemoveFast. RemoveFast is the default recommendation unless you need the list order to be stable or are removing elements inside a RangedFor iteration loop.")
-	bool Remove(const FActorRepListType& ElementToRemove)
-	{
-		return RemoveSlow(ElementToRemove);
-	}
-
 	/** Removes the element quickly but changes the list order */
 	bool RemoveFast(const FActorRepListType& ElementToRemove)
 	{
@@ -343,12 +321,6 @@ struct REPLICATIONGRAPH_API FActorRepListRefView
 	TArray<FActorRepListType>::RangedForIteratorType end()				{ return RepList.end(); }
 	TArray<FActorRepListType>::RangedForConstIteratorType end() const	{ return RepList.end(); }
 
-	UE_DEPRECATED(4.27, "IsValid is deprecated now that you don't need to call PrepareForWrite before doing operations on the list. Use IsEmpty() if you need to skip doing operations on empty lists")
-	bool IsValid() const 
-	{ 
-		return true; 
-	}
-
 	bool IsEmpty() const
 	{
 		return RepList.Num() <= 0;
@@ -365,12 +337,6 @@ struct REPLICATIONGRAPH_API FActorRepListRefView
 		RepList.Empty();
 	}
 	
-	UE_DEPRECATED(4.27, "ResetToNull was renamed to TearDown")
-	void ResetToNull()
-	{ 
-		TearDown();
-	}
-
 	int32 IndexOf(const FActorRepListType& Value) const
 	{
 		return RepList.IndexOfByKey(Value);
@@ -421,12 +387,6 @@ private:
 
 /** A read only, non owning (ref counting) view to an actor replication list: essentially a raw pointer and the category of the list. These are only created *from* FActorRepListRefView */
 
-struct UE_DEPRECATED(4.27, "Replace this struct with the new FActorRepListConstView struct.")  FActorRepListRawView : public TActorRepListViewBase<FActorRepList*>
-{
-	/** Standard ctor: make raw view from ref view */
-	REPLICATIONGRAPH_API FActorRepListRawView(const FActorRepListRefView& Source) { /*deprecated*/ }
-};
-
 #if !(UE_BUILD_SHIPPING || UE_BUILD_TEST)
 	// Intended to be called from immediate mode window while debugging
 	extern "C" DLLEXPORT void PrintRepListDetails(int32 PoolSize, int32 BlockIdx, int32 ListIdx);
@@ -435,10 +395,6 @@ struct UE_DEPRECATED(4.27, "Replace this struct with the new FActorRepListConstV
 	void PrintRepListStatsAr(int32 mode, FOutputDevice& Ar=*GLog);	
 #endif
 
-/** To be called by projects to preallocate replication lists. This isn't strictly necessary: lists will be allocated on demand as well. */
-UE_DEPRECATED(4.27, "The RepList allocator is not used anymore so preallocating is not needed.")
-REPLICATIONGRAPH_API void PreAllocateRepList(int32 ListSize, int32 NumLists);
-
 // --------------------------------------------------------------------------------------------------------------------------------------------
 // --------------------------------------------------------------------------------------------------------------------------------------------
 // Per Class/Actor Global/PerConnectino Data Structs
