ReplicationGraph
* Refactored Dependent actor lists so the child actors are stored by streaming levels
* Now dependent actors are not replicated if the connection does not have the streaming level loaded
* Added FLevelBasedActorList struct to RepGraph plugin. Holds actors based on their level (permanent or streaming)

diff --git a/Engine/Plugins/Runtime/ReplicationGraph/Source/Private/ReplicationGraph.cpp b/Engine/Plugins/Runtime/ReplicationGraph/Source/Private/ReplicationGraph.cpp
--- a/Engine/Plugins/Runtime/ReplicationGraph/Source/Private/ReplicationGraph.cpp
+++ b/Engine/Plugins/Runtime/ReplicationGraph/Source/Private/ReplicationGraph.cpp
@@ -395,6 +395,42 @@ void UReplicationGraph::RemoveConnectionGraphNode(UReplicationGraphNode* GraphNo
 	ConnectionManager->RemoveConnectionGraphNode(GraphNode);
 }
 
+UNetReplicationGraphConnection* UReplicationGraph::FindConnectionManager(UNetConnection* NetConnection) const
+{
+#if WITH_SERVER_CODE
+	check(NetConnection);
+
+	// Children do not have a connection manager, this is handled by their parent.
+	// We do not want to create connection managers for children, so redirect them.
+	if (NetConnection->GetUChildConnection() != nullptr)
+	{
+		NetConnection = ((UChildConnection*)NetConnection)->Parent;
+		UE_LOG(LogReplicationGraph, Verbose, TEXT("UReplicationGraph::FindOrAddConnectionManager was called with a child connection, redirecting to parent"));
+		check(NetConnection != nullptr);
+	}
+
+	check(NetConnection->GetDriver() == NetDriver);
+
+	for (UNetReplicationGraphConnection* ConnManager : Connections)
+	{
+		if (ConnManager->NetConnection == NetConnection)
+		{
+			return ConnManager;
+		}
+	}
+
+	for (UNetReplicationGraphConnection* ConnManager : PendingConnections)
+	{
+		if (ConnManager->NetConnection == NetConnection)
+		{
+			return ConnManager;
+		}
+	}
+#endif
+
+	return nullptr;
+}
+
 UNetReplicationGraphConnection* UReplicationGraph::FindOrAddConnectionManager(UNetConnection* NetConnection)
 {
 #if WITH_SERVER_CODE
@@ -1042,9 +1078,7 @@ int32 UReplicationGraph::ServerReplicateActors(float DeltaSeconds)
 		
 		FGatheredReplicationActorLists GatheredReplicationListsForConnection;
 
-		TSet<FName> AllVisibleLevelNames;
-		ConnectionManager->GetClientVisibleLevelNames(AllVisibleLevelNames);
-		const FConnectionGatherActorListParameters Parameters(ConnectionViewers, *ConnectionManager, AllVisibleLevelNames, FrameNum, GatheredReplicationListsForConnection);
+		const FConnectionGatherActorListParameters Parameters(ConnectionViewers, *ConnectionManager, ConnectionManager->GetCachedClientVisibleLevelNames(), FrameNum, GatheredReplicationListsForConnection);
 
 		UNetReplicationGraphConnection::FRepGraphDestructionViewerInfoArray DestructionViewersInfo;
 
@@ -1494,7 +1528,7 @@ void UReplicationGraph::ReplicateActorsForConnection(UNetConnection* NetConnecti
 		{
 			// We've exceeded the budget for this category of replication list.
 			RG_QUICK_SCOPE_CYCLE_COUNTER(NET_ReplicateActors_PartialStarvedActorList);
-			HandleStarvedActorList(PrioritizedReplicationList, ActorIdx + 1, ConnectionActorInfoMap, FrameNum);
+			HandleStarvedActorList(*ConnectionManager, PrioritizedReplicationList, ActorIdx + 1, ConnectionActorInfoMap, FrameNum);
 			NotifyConnectionSaturated(*ConnectionManager);
 			break;
 		}
@@ -1953,41 +1987,48 @@ int64 UReplicationGraph::ReplicateSingleActor(AActor* Actor, FConnectionReplicat
 	// ----------------------------
 	//	Dependent actors
 	// ----------------------------
-	const FGlobalActorReplicationInfo::FDependantListType& DependentActorList = GlobalActorInfo.GetDependentActorList();
-	if (DependentActorList.Num() > 0)
+	FGatheredReplicationActorLists ListGatherer;
+	GlobalActorInfo.GatherDependentActorLists(ConnectionManager, ListGatherer);
+	
+	if (ListGatherer.NumLists() > 0)
 	{
 		RG_QUICK_SCOPE_CYCLE_COUNTER(NET_ReplicateActors_DependentActors);
 
 		const int32 CloseFrameNum = ActorInfo.ActorChannelCloseFrameNum;
+		
+		const TArray<FActorRepListConstView>& DependentLists = ListGatherer.GetLists(EActorRepListTypeFlags::Default);
 
-		for (AActor* DependentActor : DependentActorList)
+		for (const FActorRepListConstView& DependentActorList : DependentLists)
 		{
-			repCheck(DependentActor);
+			for (AActor* DependentActor : DependentActorList)
+			{
+				repCheck(DependentActor);
 
-			FConnectionReplicationActorInfo& DependentActorConnectionInfo = ConnectionActorInfoMap.FindOrAdd(DependentActor);
-			FGlobalActorReplicationInfo& DependentActorGlobalData = GlobalActorReplicationInfoMap.Get(DependentActor);
+				FConnectionReplicationActorInfo& DependentActorConnectionInfo = ConnectionActorInfoMap.FindOrAdd(DependentActor);
+				FGlobalActorReplicationInfo& DependentActorGlobalData = GlobalActorReplicationInfoMap.Get(DependentActor);
 
-			UpdateActorChannelCloseFrameNum(DependentActor, DependentActorConnectionInfo, DependentActorGlobalData, FrameNum, NetConnection);
+				UpdateActorChannelCloseFrameNum(DependentActor, DependentActorConnectionInfo, DependentActorGlobalData, FrameNum, NetConnection);
 
-			// Dependent actor channel will stay open as long as the owning actor channel is open
-			DependentActorConnectionInfo.ActorChannelCloseFrameNum = FMath::Max<uint32>(CloseFrameNum, DependentActorConnectionInfo.ActorChannelCloseFrameNum);
+				// Dependent actor channel will stay open as long as the owning actor channel is open
+				DependentActorConnectionInfo.ActorChannelCloseFrameNum = FMath::Max<uint32>(CloseFrameNum, DependentActorConnectionInfo.ActorChannelCloseFrameNum);
 
-			if (!ReadyForNextReplication(DependentActorConnectionInfo, DependentActorGlobalData, FrameNum))
-			{
-				continue;
-			}
+				if (!ReadyForNextReplication(DependentActorConnectionInfo, DependentActorGlobalData, FrameNum))
+				{
+					continue;
+				}
 
-			if (!ensureMsgf(IsActorValidForReplication(DependentActor), TEXT("DependentActor %s (Owner: %s) not valid for replication (BeingDestroyed:%d) (IsValid:%d) (Unreachable:%d) (TearOff:%d)! Channel = %s"),
-							*DependentActor->GetFullName(), *Actor->GetFullName(),
-							DependentActor->IsActorBeingDestroyed(), IsValid(DependentActor), DependentActor->IsUnreachable(), DependentActor->GetTearOff(),
-							*DescribeSafe(DependentActorConnectionInfo.Channel)))
-			{
-				continue;
-			}
+				if (!ensureMsgf(IsActorValidForReplication(DependentActor), TEXT("DependentActor %s (Owner: %s) not valid for replication (BeingDestroyed:%d) (IsValid:%d) (Unreachable:%d) (TearOff:%d)! Channel = %s"),
+					*DependentActor->GetFullName(), *Actor->GetFullName(),
+					DependentActor->IsActorBeingDestroyed(), IsValid(DependentActor), DependentActor->IsUnreachable(), DependentActor->GetTearOff(),
+					*DescribeSafe(DependentActorConnectionInfo.Channel)))
+				{
+					continue;
+				}
 
-			//UE_LOG(LogReplicationGraph, Display, TEXT("DependentActor %s %s. NextReplicationFrameNum: %d. FrameNum: %d. ForceNetUpdateFrame: %d. LastRepFrameNum: %d."), *DependentActor->GetPathName(), *NetConnection->GetName(), DependentActorConnectionInfo.NextReplicationFrameNum, FrameNum, DependentActorGlobalData.ForceNetUpdateFrame, DependentActorConnectionInfo.LastRepFrameNum);
-			BitsWritten += ReplicateSingleActor(DependentActor, DependentActorConnectionInfo, DependentActorGlobalData, ConnectionActorInfoMap, ConnectionManager, FrameNum);
-		}					
+				//UE_LOG(LogReplicationGraph, Display, TEXT("DependentActor %s %s. NextReplicationFrameNum: %d. FrameNum: %d. ForceNetUpdateFrame: %d. LastRepFrameNum: %d."), *DependentActor->GetPathName(), *NetConnection->GetName(), DependentActorConnectionInfo.NextReplicationFrameNum, FrameNum, DependentActorGlobalData.ForceNetUpdateFrame, DependentActorConnectionInfo.LastRepFrameNum);
+				BitsWritten += ReplicateSingleActor(DependentActor, DependentActorConnectionInfo, DependentActorGlobalData, ConnectionActorInfoMap, ConnectionManager, FrameNum);
+			}
+		}
 	}
 
 	// Optional budget for actor discovery traffic
@@ -2006,27 +2047,30 @@ int64 UReplicationGraph::ReplicateSingleActor(AActor* Actor, FConnectionReplicat
 #endif //WITH_SERVER_CODE
 }
 
-void UReplicationGraph::HandleStarvedActorList(const FPrioritizedRepList& List, int32 StartIdx, FPerConnectionActorInfoMap& ConnectionActorInfoMap, uint32 FrameNum)
+void UReplicationGraph::HandleStarvedActorList(const UNetReplicationGraphConnection& RepGraphConnection, const FPrioritizedRepList& List, int32 StartIdx, FPerConnectionActorInfoMap& ConnectionActorInfoMap, uint32 FrameNum)
 {
+	FGatheredReplicationActorLists DependentListContainer;
+
 	for (int32 ActorIdx=StartIdx; ActorIdx < List.Items.Num(); ++ActorIdx)
 	{
 		const FPrioritizedRepList::FItem& RepItem = List.Items[ActorIdx];
-		FConnectionReplicationActorInfo& ActorInfo = *RepItem.ConnectionData;
+		const uint32 CloseFrameNum = RepItem.ConnectionData->ActorChannelCloseFrameNum;
 
 		// Update dependent actor's timeout frame
 		FGlobalActorReplicationInfo& GlobalActorInfo = GlobalActorReplicationInfoMap.Get(RepItem.Actor);
 
-		const FGlobalActorReplicationInfo::FDependantListType& DependentActorList = GlobalActorInfo.GetDependentActorList();
-
-		if (DependentActorList.Num() > 0)
+		
+		GlobalActorInfo.GatherDependentActorLists(RepGraphConnection, DependentListContainer);
+		const TArray<FActorRepListConstView>& DependentLists = DependentListContainer.GetLists(EActorRepListTypeFlags::Default);
+		for (const FActorRepListConstView& DependentList : DependentLists)
 		{
-			const uint32 CloseFrameNum = ActorInfo.ActorChannelCloseFrameNum;
-			for (AActor* DependentActor : DependentActorList)
+			for (AActor* DependentActor : DependentList)
 			{
 				FConnectionReplicationActorInfo& DependentActorConnectionInfo = ConnectionActorInfoMap.FindOrAdd(DependentActor);
 				DependentActorConnectionInfo.ActorChannelCloseFrameNum = FMath::Max<uint32>(CloseFrameNum, DependentActorConnectionInfo.ActorChannelCloseFrameNum);
 			}
 		}
+		DependentListContainer.Reset();
 	}
 }
 
@@ -2638,9 +2682,26 @@ bool UNetReplicationGraphConnection::PrepareForReplication()
 		CurChild->ViewTarget = CurChild->PlayerController ? CurChild->PlayerController->GetViewTarget() : ToRawPtr(CurChild->OwningActor);
 	}
 
+	BuildVisibleLevels();
+
 	return (NetConnection->GetConnectionState() != USOCK_Closed) && (NetConnection->ViewTarget != nullptr) && bConnectionHasCorrectWorld;
 }
 
+void UNetReplicationGraphConnection::BuildVisibleLevels()
+{
+	CachedVisibleLevels.Reset();
+
+	CachedVisibleLevels.Append(NetConnection->ClientVisibleLevelNames);
+	for (const UNetConnection* Child : NetConnection->Children)
+	{
+		if (Child)
+		{
+			// For sets, we don't have to worry about uniqueness due to the nature of the data structure.
+			CachedVisibleLevels.Append(Child->ClientVisibleLevelNames);
+		}
+	}
+}
+
 void UNetReplicationGraphConnection::NotifyAddDestructionInfo(FActorDestructionInfo* DestructInfo)
 {
 #if WITH_SERVER_CODE
@@ -3084,7 +3145,7 @@ void FStreamingLevelActorListCollection::AddActor(const FNewReplicatedActorInfo&
 	Item->ReplicationActorList.Add(ActorInfo.Actor);
 }
 
-bool FStreamingLevelActorListCollection::RemoveActor(const FNewReplicatedActorInfo& ActorInfo, bool bWarnIfNotFound, UReplicationGraphNode* Outer)
+bool FStreamingLevelActorListCollection::RemoveActor(const FNewReplicatedActorInfo& ActorInfo, bool bWarnIfNotFound, UObject* Outer)
 {
 	bool bRemovedSomething = false;
 	for (FStreamingLevelActors& StreamingList : StreamingLevelLists)
@@ -3107,7 +3168,7 @@ bool FStreamingLevelActorListCollection::RemoveActor(const FNewReplicatedActorIn
 	return bRemovedSomething;
 }
 
-bool FStreamingLevelActorListCollection::RemoveActorFast(const FNewReplicatedActorInfo& ActorInfo, UReplicationGraphNode* Outer)
+bool FStreamingLevelActorListCollection::RemoveActorFast(const FNewReplicatedActorInfo& ActorInfo, UObject* Outer)
 {
 	bool bRemovedSomething = false;
 	for (FStreamingLevelActors& StreamingList : StreamingLevelLists)
@@ -3129,7 +3190,7 @@ void FStreamingLevelActorListCollection::Reset()
 	}
 }
 
-void FStreamingLevelActorListCollection::Gather(const FConnectionGatherActorListParameters& Params)
+void FStreamingLevelActorListCollection::Gather(const FConnectionGatherActorListParameters& Params) const
 {
 	for (const FStreamingLevelActors& StreamingList : StreamingLevelLists)
 	{
@@ -3144,6 +3205,40 @@ void FStreamingLevelActorListCollection::Gather(const FConnectionGatherActorList
 	}
 }
 
+void FStreamingLevelActorListCollection::Gather(const UNetReplicationGraphConnection& ConnectionManager, FGatheredReplicationActorLists& OutGatheredList) const
+{
+	const TSet<FName>& ClientLevelNames = ConnectionManager.GetCachedClientVisibleLevelNames();
+
+	for (const FStreamingLevelActors& StreamingList : StreamingLevelLists)
+	{
+		if (ClientLevelNames.Contains(StreamingList.StreamingLevelName))
+		{
+			OutGatheredList.AddReplicationActorList(StreamingList.ReplicationActorList);
+		}
+	}
+}
+
+void FStreamingLevelActorListCollection::AppendAllLists(FGatheredReplicationActorLists& OutGatheredList) const
+{
+	for (const FStreamingLevelActors& StreamingList : StreamingLevelLists)
+	{
+		OutGatheredList.AddReplicationActorList(StreamingList.ReplicationActorList);
+	}
+}
+
+bool FStreamingLevelActorListCollection::Contains(const FNewReplicatedActorInfo& ActorInfo) const
+{
+	for (const FStreamingLevelActors& StreamingList : StreamingLevelLists)
+	{
+		if (StreamingList.StreamingLevelName == ActorInfo.StreamingLevelName)
+		{
+			return StreamingList.ReplicationActorList.Contains(ActorInfo.Actor);
+		}
+	}
+
+	return false;
+}
+
 void FStreamingLevelActorListCollection::DeepCopyFrom(const FStreamingLevelActorListCollection& Source)
 {
 	StreamingLevelLists.Reset();
@@ -3738,7 +3833,7 @@ void UReplicationGraphNode_DynamicSpatialFrequency::GatherActorListsForConnectio
 					AActor* Actor = Item.Actor;
 					FGlobalActorReplicationInfo& GlobalInfo = *Item.GlobalInfo;
 				
-					CalcFrequencyForActor(Actor, RepGraph, NetConnection, GlobalInfo, ConnectionActorInfoMap, MySettings, Params.Viewers, FrameNum, idx);
+					CalcFrequencyForActor(Actor, RepGraph, Params.ConnectionManager, GlobalInfo, ConnectionActorInfoMap, MySettings, Params.Viewers, FrameNum, idx);
 				}
 
 				SortedReplicationList.Sort();
@@ -3886,10 +3981,18 @@ FORCEINLINE uint32 CalcDynamicReplicationPeriod(const float FinalPCT, const uint
 static TArray<FColor> DynamicSpatialFrequencyDebugColorArray = { FColor::Red, FColor::Green, FColor::Blue, FColor::Cyan, FColor::Orange, FColor::Purple };
 
 void UReplicationGraphNode_DynamicSpatialFrequency::CalcFrequencyForActor(AActor* Actor, UReplicationGraph* RepGraph, UNetConnection* NetConnection, FGlobalActorReplicationInfo& GlobalInfo, FPerConnectionActorInfoMap& ConnectionMap, FSettings& MySettings, const FNetViewerArray& Viewers, const uint32 FrameNum, int32 ExistingItemIndex)
+{
+	UNetReplicationGraphConnection* RepGraphConnection = RepGraph->FindConnectionManager(NetConnection);
+	CalcFrequencyForActor(Actor, RepGraph, *RepGraphConnection, GlobalInfo, ConnectionMap, MySettings, Viewers, FrameNum, ExistingItemIndex);
+}
+
+void UReplicationGraphNode_DynamicSpatialFrequency::CalcFrequencyForActor(AActor* Actor, UReplicationGraph* RepGraph, const UNetReplicationGraphConnection& RepGraphConnection, FGlobalActorReplicationInfo& GlobalInfo, FPerConnectionActorInfoMap& ConnectionMap, FSettings& MySettings, const FNetViewerArray& Viewers, const uint32 FrameNum, int32 ExistingItemIndex)
 {
 #if WITH_SERVER_CODE
 	FConnectionReplicationActorInfo& ConnectionInfo = ConnectionMap.FindOrAdd(Actor);
 
+	UNetConnection* NetConnection = RepGraphConnection.NetConnection;
+
 	// If we need to filter out the actor and it is already in the SortedReplicationList, we need to remove it (instead of just skipping/returning).
 	auto RemoveExistingItem = [&ExistingItemIndex, this]()
 	{
@@ -3957,6 +4060,8 @@ void UReplicationGraphNode_DynamicSpatialFrequency::CalcFrequencyForActor(AActor
 
 	const bool ActorSupportsFastShared = (GlobalInfo.Settings.FastSharedReplicationFunc != nullptr);
 	TArrayView<FSpatializationZone>& ZoneList = ActorSupportsFastShared ? MySettings.ZoneSettings : MySettings.ZoneSettings_NonFastSharedActors;
+
+	FGatheredReplicationActorLists DependentListContainer;
 		
 	for (int32 ZoneIdx=0; ZoneIdx < ZoneList.Num(); ++ZoneIdx)
 	{
@@ -3995,12 +4100,18 @@ void UReplicationGraphNode_DynamicSpatialFrequency::CalcFrequencyForActor(AActor
 				// Update actor timeout frame here in case we get starved and can't actually replicate before then
 				ConnectionInfo.ActorChannelCloseFrameNum = FMath::Max<uint32>(ConnectionInfo.ActorChannelCloseFrameNum, ConnectionInfo.NextReplicationFrameNum + 1);
 
-				const FGlobalActorReplicationInfo::FDependantListType& DependentActorList = GlobalInfo.GetDependentActorList();
-				for (AActor* DependentActor : DependentActorList)
+				// Update the dependent actor's timeout frame too
+				GlobalInfo.GatherDependentActorLists(RepGraphConnection, DependentListContainer);
+				const TArray<FActorRepListConstView>& DependentLists = DependentListContainer.GetLists(EActorRepListTypeFlags::Default);
+				for (const FActorRepListConstView& DependentList : DependentLists)
 				{
-					FConnectionReplicationActorInfo& DependentActorConnectionInfo = ConnectionMap.FindOrAdd(DependentActor);
-					DependentActorConnectionInfo.ActorChannelCloseFrameNum = FMath::Max<uint32>(ConnectionInfo.ActorChannelCloseFrameNum, ConnectionInfo.NextReplicationFrameNum + 1);
+					for (AActor* DependentActor : DependentList)
+					{
+						FConnectionReplicationActorInfo& DependentActorConnectionInfo = ConnectionMap.FindOrAdd(DependentActor);
+						DependentActorConnectionInfo.ActorChannelCloseFrameNum = FMath::Max<uint32>(ConnectionInfo.ActorChannelCloseFrameNum, ConnectionInfo.NextReplicationFrameNum + 1);
+					}
 				}
+				DependentListContainer.Reset();
 
 				// Calc Replication Period for FastShared replication
 				if (ActorSupportsFastShared && ZoneInfo.FastPath_MinRepPeriod > 0)
@@ -4111,7 +4222,7 @@ void UReplicationGraphNode_DynamicSpatialFrequency::GatherActors(const FActorRep
 
 		FGlobalActorReplicationInfo& GlobalInfo = GlobalMap.Get(Actor);
 
-		CalcFrequencyForActor(Actor, RepGraph, NetConnection, GlobalInfo, ConnectionMap, MySettings, Params.Viewers, FrameNum, INDEX_NONE);
+		CalcFrequencyForActor(Actor, RepGraph, Params.ConnectionManager, GlobalInfo, ConnectionMap, MySettings, Params.Viewers, FrameNum, INDEX_NONE);
 	}
 }
 
diff --git a/Engine/Plugins/Runtime/ReplicationGraph/Source/Private/ReplicationGraphDebugging.cpp b/Engine/Plugins/Runtime/ReplicationGraph/Source/Private/ReplicationGraphDebugging.cpp
--- a/Engine/Plugins/Runtime/ReplicationGraph/Source/Private/ReplicationGraphDebugging.cpp
+++ b/Engine/Plugins/Runtime/ReplicationGraph/Source/Private/ReplicationGraphDebugging.cpp
@@ -303,7 +303,12 @@ void AReplicationGraphDebugActor::PrintAllActorInfo(FString MatchString)
 
 	for (auto ClassRepInfoIt = ReplicationGraph->GlobalActorReplicationInfoMap.CreateClassMapIterator(); ClassRepInfoIt; ++ClassRepInfoIt)
 	{
-		UClass* Class = CastChecked<UClass>(ClassRepInfoIt.Key().ResolveObjectPtr());
+		UClass* Class = Cast<UClass>(ClassRepInfoIt.Key().ResolveObjectPtr());
+		if (!Class)
+		{
+			continue;
+		}
+
 		const FClassReplicationInfo& ClassInfo = ClassRepInfoIt.Value();
 
 		if (!Matches(Class))
@@ -967,12 +972,26 @@ void FGlobalActorReplicationInfo::LogDebugString(FOutputDevice& Ar) const
 	Ar.Logf(TEXT("  LastPreReplicationFrame: %d. ForceNetUpdateFrame: %d. WorldLocation: %s. bWantsToBeDormant %d. LastFlushNetDormancyFrame: %d"), LastPreReplicationFrame, ForceNetUpdateFrame, *WorldLocation.ToString(), bWantsToBeDormant, LastFlushNetDormancyFrame);
 	Ar.Logf(TEXT("  Settings: %s"), *Settings.BuildDebugStringDelta());
 
-	if (DependentActorList.Num() > 0)
+	TArray<AActor*> DependentActors;
+	DependentActorList.GetAllActors(DependentActors);
+
+	if (DependentActors.Num() > 0)
 	{
 		FString DependentActorStr = TEXT("DependentActors: ");
-		for (FActorRepListType Actor : DependentActorList)
+		for (AActor* DependentActor : DependentActors)
+		{
+			DependentActorStr += GetActorRepListTypeDebugString(DependentActor) + ' ';
+		}
+
+		Ar.Logf(TEXT("  %s"), *DependentActorStr);
+	}
+
+	if (ParentActorList.Num() > 0)
+	{
+		FString DependentActorStr = TEXT("IsADependentOf: ");
+		for (AActor* ParentDependentActor : ParentActorList)
 		{
-			DependentActorStr += GetActorRepListTypeDebugString(Actor) + ' ';
+			DependentActorStr += GetActorRepListTypeDebugString(ParentDependentActor) + ' ';
 		}
 
 		Ar.Logf(TEXT("  %s"), *DependentActorStr);
diff --git a/Engine/Plugins/Runtime/ReplicationGraph/Source/Private/ReplicationGraphTypes.cpp b/Engine/Plugins/Runtime/ReplicationGraph/Source/Private/ReplicationGraphTypes.cpp
--- a/Engine/Plugins/Runtime/ReplicationGraph/Source/Private/ReplicationGraphTypes.cpp
+++ b/Engine/Plugins/Runtime/ReplicationGraph/Source/Private/ReplicationGraphTypes.cpp
@@ -445,6 +445,18 @@ void TActorListAllocator<NumListsPerBlock, MaxNumPools>::LogDetails(int32 PoolSi
 	}
 }
 
+int32 FGlobalActorReplicationInfoMap::Remove(const FActorRepListType& RemovedActor)
+{
+	// Clean the references to the removed actor from its dependent link
+	if (FGlobalActorReplicationInfo* RemovedActorInfo = Find(RemovedActor))
+	{
+		RemoveAllActorDependencies(RemovedActor, RemovedActorInfo);
+	}
+
+	return ActorMap.Remove(RemovedActor);
+}
+
+
 void FGlobalActorReplicationInfoMap::AddDependentActor(AActor* Parent, AActor* Child, FGlobalActorReplicationInfoMap::EWarnFlag WarnFlag)
 {
 	/**
@@ -467,12 +479,12 @@ void FGlobalActorReplicationInfoMap::AddDependentActor(AActor* Parent, AActor* C
 		bool bChildIsAlreadyDependant(false);
 		if (FGlobalActorReplicationInfo* ParentInfo = Find(Parent))
 		{
-			bChildIsAlreadyDependant = ParentInfo->DependentActorList.Find(Child) != INDEX_NONE;
+			bChildIsAlreadyDependant = ParentInfo->DependentActorList.Contains(Child);
 			if (bChildIsAlreadyDependant == false)
 			{
 				if (IsActorValidForReplicationGather(Child))
 				{
-					ParentInfo->DependentActorList.Add(Child);
+					ParentInfo->DependentActorList.AddNetworkActor(Child);
 				}
 			}
 		}
@@ -597,3 +609,93 @@ void FActorRepListStatCollector::VisitExplicitStreamingLevelList(FName ListOwner
 	StreamingLevelStats.NumBytes += ListBytes;
 	StreamingLevelStats.MaxListSize = FMath::Max(StreamingLevelStats.MaxListSize, ListSize);
 }
+
+void FLevelBasedActorList::AddNetworkActor(AActor* NetActor)
+{
+	FNewReplicatedActorInfo ActorInfo(NetActor);
+
+	if (ActorInfo.StreamingLevelName == NAME_None)
+	{
+		PermanentLevelActors.Add(NetActor);
+	}
+	else
+	{
+		StreamingLevelActors.AddActor(ActorInfo);
+	}
+}
+
+bool FLevelBasedActorList::RemoveNetworkActor(const FNewReplicatedActorInfo& ActorInfo)
+{
+	if (ActorInfo.StreamingLevelName == NAME_None)
+	{
+		return PermanentLevelActors.RemoveFast(ActorInfo.Actor);
+	}
+	else
+	{
+		return StreamingLevelActors.RemoveActorFast(ActorInfo);
+	}
+}
+
+bool FLevelBasedActorList::RemoveNetworkActorOrdered(AActor* NetActor)
+{
+	FNewReplicatedActorInfo ActorInfo(NetActor);
+
+	if (ActorInfo.StreamingLevelName == NAME_None)
+	{
+		return PermanentLevelActors.RemoveSlow(NetActor);
+	}
+	else
+	{
+		return StreamingLevelActors.RemoveActor(ActorInfo, false);
+	}
+}
+
+bool FLevelBasedActorList::Contains(AActor* NetActor) const
+{
+	FNewReplicatedActorInfo ActorInfo(NetActor);
+
+	if (ActorInfo.StreamingLevelName == NAME_None)
+	{
+		return PermanentLevelActors.Contains(NetActor);
+	}
+	else
+	{
+		return StreamingLevelActors.Contains(ActorInfo);
+	}
+}
+
+void FLevelBasedActorList::Reset()
+{
+	PermanentLevelActors.Reset();
+	StreamingLevelActors.Reset();
+}
+
+void FLevelBasedActorList::Gather(const FConnectionGatherActorListParameters& Params) const
+{
+	Params.OutGatheredReplicationLists.AddReplicationActorList(PermanentLevelActors);
+	StreamingLevelActors.Gather(Params);
+}
+
+void FLevelBasedActorList::Gather(const UNetReplicationGraphConnection& ConnectionManager, FGatheredReplicationActorLists& OutGatheredList) const
+{
+	OutGatheredList.AddReplicationActorList(PermanentLevelActors);
+	StreamingLevelActors.Gather(ConnectionManager, OutGatheredList);
+}
+
+void FLevelBasedActorList::AppendAllLists(FGatheredReplicationActorLists& OutGatheredList) const
+{
+	OutGatheredList.AddReplicationActorList(PermanentLevelActors);
+	StreamingLevelActors.AppendAllLists(OutGatheredList);
+}
+
+void FLevelBasedActorList::GetAllActors(TArray<AActor*>& OutAllActors) const
+{
+	PermanentLevelActors.AppendToTArray(OutAllActors);
+	StreamingLevelActors.GetAll_Debug(OutAllActors);
+}
+
+void FLevelBasedActorList::CountBytes(FArchive& Ar) const
+{
+	PermanentLevelActors.CountBytes(Ar);
+	StreamingLevelActors.CountBytes(Ar);
+}
diff --git a/Engine/Plugins/Runtime/ReplicationGraph/Source/Public/ReplicationGraph.h b/Engine/Plugins/Runtime/ReplicationGraph/Source/Public/ReplicationGraph.h
--- a/Engine/Plugins/Runtime/ReplicationGraph/Source/Public/ReplicationGraph.h
+++ b/Engine/Plugins/Runtime/ReplicationGraph/Source/Public/ReplicationGraph.h
@@ -67,7 +67,7 @@ struct FReplicationGraphDestructionSettings;
 
 typedef TObjectKey<class UNetReplicationGraphConnection> FRepGraphConnectionKey;
 
-#define DO_ENABLE_REPGRAPH_DEBUG_ACTOR !(UE_BUILD_SHIPPING || UE_BUILD_TEST)
+#define DO_ENABLE_REPGRAPH_DEBUG_ACTOR !(UE_BUILD_SHIPPING)
 
 UCLASS(abstract, transient, config=Engine)
 class REPLICATIONGRAPH_API UReplicationGraphNode : public UObject
@@ -169,49 +169,6 @@ protected:
 
 // -----------------------------------
 
-
-struct REPLICATIONGRAPH_API FStreamingLevelActorListCollection
-{
-	void AddActor(const FNewReplicatedActorInfo& ActorInfo);
-	bool RemoveActor(const FNewReplicatedActorInfo& ActorInfo, bool bWarnIfNotFound, UReplicationGraphNode* Outer);
-	bool RemoveActorFast(const FNewReplicatedActorInfo& ActorInfo, UReplicationGraphNode* Outer);
-	void Reset();
-	void Gather(const FConnectionGatherActorListParameters& Params);
-	void DeepCopyFrom(const FStreamingLevelActorListCollection& Source);
-	void GetAll_Debug(TArray<FActorRepListType>& OutArray) const;
-	void Log(FReplicationGraphDebugInfo& DebugInfo) const;
-	int32 NumLevels() const { return StreamingLevelLists.Num(); }
-	void TearDown();
-
-	struct FStreamingLevelActors
-	{
-		FStreamingLevelActors(FName InName) : StreamingLevelName(InName) 
-		{ 
-			repCheck(InName != NAME_None); 
-		}
-
-		FName StreamingLevelName;
-		FActorRepListRefView ReplicationActorList;
-		bool operator==(const FName& InName) const { return InName == StreamingLevelName; };
-	};
-
-	/** Lists for streaming levels. Actors that "came from" streaming levels go here. These lists are only returned if the connection has their streaming level loaded. */
-	static const int32 NumInlineAllocations = 4;
-	TArray<FStreamingLevelActors, TInlineAllocator<NumInlineAllocations>> StreamingLevelLists;
-
-	void CountBytes(FArchive& Ar) const
-	{
-		StreamingLevelLists.CountBytes(Ar);
-
-		for (const FStreamingLevelActors& List : StreamingLevelLists)
-		{
-			List.ReplicationActorList.CountBytes(Ar);
-		}
-	}
-};
-
-// -----------------------------------
-
 /** A Node that contains ReplicateActorLists. This contains 1 "base" list and a TArray of lists that are conditioned on a streaming level being loaded. */
 UCLASS()
 class REPLICATIONGRAPH_API UReplicationGraphNode_ActorList : public UReplicationGraphNode
@@ -434,7 +391,10 @@ protected:
 	virtual void GatherActors(const FActorRepListRefView& RepList, FGlobalActorReplicationInfoMap& GlobalMap, FPerConnectionActorInfoMap& ConnectionMap, const FConnectionGatherActorListParameters& Params, UNetConnection* NetConnection);
 	virtual void GatherActors_DistanceOnly(const FActorRepListRefView& RepList, FGlobalActorReplicationInfoMap& GlobalMap, FPerConnectionActorInfoMap& ConnectionMap, const FConnectionGatherActorListParameters& Params);
 
+	UE_DEPRECATED(5.2, "This function was replaced with a version that needs to receive a UNetReplicationGraphConnection instead of a NetConnection")
 	void CalcFrequencyForActor(AActor* Actor, UReplicationGraph* RepGraph, UNetConnection* NetConnection, FGlobalActorReplicationInfo& GlobalInfo, FPerConnectionActorInfoMap& ConnectionMap, FSettings& MySettings, const FNetViewerArray& Viewers, const uint32 FrameNum, int32 ExistingItemIndex);
+
+	void CalcFrequencyForActor(AActor* Actor, UReplicationGraph* RepGraph, const UNetReplicationGraphConnection& RepGraphConnection, FGlobalActorReplicationInfo& GlobalInfo, FPerConnectionActorInfoMap& ConnectionMap, FSettings& MySettings, const FNetViewerArray& Viewers, const uint32 FrameNum, int32 ExistingItemIndex);
 };
 
 
@@ -1076,6 +1036,8 @@ public:
 		return (uint16)FramesBetweenUpdates;
 	}
 
+	UNetReplicationGraphConnection* FindConnectionManager(UNetConnection* NetConnection) const;
+
 protected:
 
 	virtual void InitializeForWorld(UWorld* World);
@@ -1093,7 +1055,11 @@ protected:
 
 	UNetReplicationGraphConnection* FindOrAddConnectionManager(UNetConnection* NetConnection);
 
-	void HandleStarvedActorList(const FPrioritizedRepList& List, int32 StartIdx, FPerConnectionActorInfoMap& ConnectionActorInfoMap, uint32 FrameNum);
+	UE_DEPRECATED(5.2, "This function has been replaced with a version that needs to receive a UNetReplicationGraphConnection")
+	void HandleStarvedActorList(const FPrioritizedRepList& List, int32 StartIdx, FPerConnectionActorInfoMap& ConnectionActorInfoMap, uint32 FrameNum) {}
+
+	/** Sets the next timeout frame for the actors in the list along with their dependent actors */
+	void HandleStarvedActorList(const UNetReplicationGraphConnection& RepGraphConnection, const FPrioritizedRepList& List, int32 StartIdx, FPerConnectionActorInfoMap& ConnectionActorInfoMap, uint32 FrameNum);
 
 	/** How long, in frames, without replicating before an actor channel is closed on a connection. This is a global value added to the individual actor's ActorChannelFrameTimeout */
 	uint32 GlobalActorChannelFrameNumTimeout;
@@ -1303,6 +1269,9 @@ public:
 	/** Generates a set of all the visible level names for this connection and its subconnections (if any) */
 	virtual void GetClientVisibleLevelNames(TSet<FName>& OutLevelNames) const;
 
+	/** Returns the list of visible levels cached at the start of the gather phase. */
+	const TSet<FName>& GetCachedClientVisibleLevelNames() const { return CachedVisibleLevels; }
+
 	FActorRepListRefView& GetPrevDormantActorListForNode(const UReplicationGraphNode* GridNode);
 
 	void RemoveActorFromAllPrevDormantActorLists(AActor* InActor);
@@ -1358,6 +1327,8 @@ private:
 
 	void SetActorNotDormantOnConnection(AActor* InActor);
 
+	void BuildVisibleLevels();
+
 	UPROPERTY()
 	TArray<TObjectPtr<UReplicationGraphNode>> ConnectionGraphNodes;
 
@@ -1409,6 +1380,9 @@ private:
 
 	/** Set used to guard against double adds into PendingDormantDestructList */
 	TSet<FNetworkGUID> TrackedDormantDestructionInfos;
+
+	/** List of level names visible to this connection. Valid during the Gather phase */
+	TSet<FName> CachedVisibleLevels;
 };
 
 // --------------------------------------------------------------------------------------------------------------------------------------------
diff --git a/Engine/Plugins/Runtime/ReplicationGraph/Source/Public/ReplicationGraphTypes.h b/Engine/Plugins/Runtime/ReplicationGraph/Source/Public/ReplicationGraphTypes.h
--- a/Engine/Plugins/Runtime/ReplicationGraph/Source/Public/ReplicationGraphTypes.h
+++ b/Engine/Plugins/Runtime/ReplicationGraph/Source/Public/ReplicationGraphTypes.h
@@ -28,8 +28,11 @@ class UActorChannel;
 class UNetConnection;
 class UNetReplicationGraphConnection;
 class UReplicationGraph;
+class UReplicationGraphNode;
 
 struct FActorDestructionInfo;
+struct FReplicationGraphDebugInfo;
+struct FConnectionGatherActorListParameters;
 
 REPLICATIONGRAPH_API DECLARE_LOG_CATEGORY_EXTERN( LogReplicationGraph, Log, All );
 
@@ -395,6 +398,195 @@ private:
 	void PrintRepListStatsAr(int32 mode, FOutputDevice& Ar=*GLog);	
 #endif
 
+
+// --------------------------------------------------------------------------------------------------------------------------------------------
+// --------------------------------------------------------------------------------------------------------------------------------------------
+// Gathering Parameters
+// --------------------------------------------------------------------------------------------------------------------------------------------
+// --------------------------------------------------------------------------------------------------------------------------------------------
+
+// This represents "the list of gathered lists". This is what we push down the Replication Graph and nodes will either Push/Pop List Categories or will add their Replication Lists.
+struct REPLICATIONGRAPH_API FGatheredReplicationActorLists
+{
+	void AddReplicationActorList(const FActorRepListRefView& List, EActorRepListTypeFlags Flags = EActorRepListTypeFlags::Default)
+	{
+#if !(UE_BUILD_SHIPPING || UE_BUILD_TEST)
+		if (CVar_RepGraph_Verify)
+			List.VerifyContents_Slow();
+#endif
+		if (List.Num() > 0)
+		{
+			ReplicationLists[(uint32)Flags].Emplace(FActorRepListConstView(List));
+			CachedNum++;
+		}
+	}
+
+	FORCEINLINE void Reset()
+	{
+		for (uint32 i = (uint32)EActorRepListTypeFlags::Default; i < (uint32)EActorRepListTypeFlags::Max; ++i)
+		{
+			ReplicationLists[i].Reset();
+		}
+		CachedNum = 0;
+	}
+	FORCEINLINE int32 NumLists() const
+	{
+		return CachedNum;
+	}
+
+	FORCEINLINE const TArray<FActorRepListConstView>& GetLists(EActorRepListTypeFlags ListFlags) const
+	{
+		return ReplicationLists[(uint32)ListFlags];
+	}
+	FORCEINLINE bool ContainsLists(EActorRepListTypeFlags Flags) const
+	{
+		return ReplicationLists[(uint32)Flags].Num() > 0;
+	}
+
+private:
+
+	TStaticArray< TArray<FActorRepListConstView>, (uint32)EActorRepListTypeFlags::Max > ReplicationLists;
+	int32 CachedNum = 0;
+};
+
+// --------------------------------------------------------------------------------------------------------------------------------------------
+// --------------------------------------------------------------------------------------------------------------------------------------------
+// NewReplicatedActorInfo
+// --------------------------------------------------------------------------------------------------------------------------------------------
+// --------------------------------------------------------------------------------------------------------------------------------------------
+
+/** This is the struct we use to push new replication actors into the graph. "New" doesn't mean "newly spawned" it means "new to the graph". FIXME: Please suggest a better name! */
+struct FNewReplicatedActorInfo
+{
+	explicit FNewReplicatedActorInfo(const FActorRepListType& InActor) : Actor(InActor), Class(InActor->GetClass())
+	{
+		StreamingLevelName = GetStreamingLevelNameOfActor(Actor);
+	}
+
+	AActor* GetActor() const { return Actor; }
+
+	REPLICATIONGRAPH_API static FName GetStreamingLevelNameOfActor(const AActor* Actor);
+
+	FActorRepListType Actor;
+	FName StreamingLevelName;
+	UClass* Class;
+};
+
+// --------------------------------------------------------------------------------------------------------------------------------------------
+// --------------------------------------------------------------------------------------------------------------------------------------------
+// Level based actor lists 
+// --------------------------------------------------------------------------------------------------------------------------------------------
+// --------------------------------------------------------------------------------------------------------------------------------------------
+
+/**
+* Holds a list of actors grouped by their streaming level
+*/
+struct REPLICATIONGRAPH_API FStreamingLevelActorListCollection
+{
+	void AddActor(const FNewReplicatedActorInfo& ActorInfo);
+
+	/** 
+	* Removes the actor but keep the actor order list intact 
+	* 
+	* @param ActorInfo  The actor to be removed
+	* @param bWarnIfNotFound When true we log a warning when the actor was not in any list
+	* @param Outer Optional pointer to the owner of the list if you want to output the owner name in any log messages
+	*/
+	bool RemoveActor(const FNewReplicatedActorInfo& ActorInfo, bool bWarnIfNotFound, UObject* Outer=nullptr);
+
+	/**
+	* Removes the actor but using the faster swap operation that change the ordering of the list
+	*
+	* @param ActorInfo  The actor to be removed
+	* @param Outer Optional pointer to the owner of the list if you want to output the owner name in any log messages
+	*/
+	bool RemoveActorFast(const FNewReplicatedActorInfo& ActorInfo, UObject* Outer=nullptr);
+
+	void Reset();
+
+	/** Add to the Params' gather list the list of actors for levels that are visible to that connection */
+	void Gather(const FConnectionGatherActorListParameters& Params) const;
+
+	/** Add to the passed gather list the list of actors that are visible to the connection */
+	void Gather(const UNetReplicationGraphConnection& ConnectionManager, FGatheredReplicationActorLists& OutGatheredList) const;
+
+	/** Add to the passed gather list all list of actors held */
+	void AppendAllLists(FGatheredReplicationActorLists& OutGatheredList) const;
+
+	/** Is this actor held in one of the actor lists */
+	bool Contains(const FNewReplicatedActorInfo& ActorInfo) const;
+
+	void DeepCopyFrom(const FStreamingLevelActorListCollection& Source);
+	void GetAll_Debug(TArray<FActorRepListType>& OutArray) const;
+	void Log(FReplicationGraphDebugInfo& DebugInfo) const;
+	
+	int32 NumLevels() const { return StreamingLevelLists.Num(); }
+	
+	void TearDown();
+
+	struct FStreamingLevelActors
+	{
+		FStreamingLevelActors(FName InName) : StreamingLevelName(InName)  { repCheck(InName != NAME_None); }
+
+		FName StreamingLevelName;
+		FActorRepListRefView ReplicationActorList;
+		bool operator==(const FName& InName) const { return InName == StreamingLevelName; };
+	};
+
+	/** Lists for streaming levels. Actors that "came from" streaming levels go here. These lists are only returned if the connection has their streaming level loaded. */
+	static const int32 NumInlineAllocations = 4;
+	TArray<FStreamingLevelActors, TInlineAllocator<NumInlineAllocations>> StreamingLevelLists;
+
+	void CountBytes(FArchive& Ar) const
+	{
+		StreamingLevelLists.CountBytes(Ar);
+
+		for (const FStreamingLevelActors& List : StreamingLevelLists)
+		{
+			List.ReplicationActorList.CountBytes(Ar);
+		}
+	}
+};
+
+
+/**
+* Holds a list of actors grouped by the permanent level or their streaming level
+*/
+struct REPLICATIONGRAPH_API FLevelBasedActorList
+{
+	void AddNetworkActor(AActor* NetActor);
+	bool RemoveNetworkActor(const FNewReplicatedActorInfo& NetActor);
+	bool RemoveNetworkActor(AActor* NetActor)
+	{
+		FNewReplicatedActorInfo ActorInfo(NetActor);
+		return RemoveNetworkActor(ActorInfo);
+	}
+
+	bool RemoveNetworkActorOrdered(AActor* NetActor);
+
+	void Reset();
+
+	bool Contains(AActor* NetActor) const;
+
+	/** Gather the list of actors visible to this connection and add them to the gather list */
+	void Gather(const FConnectionGatherActorListParameters& Params) const;
+	void Gather(const UNetReplicationGraphConnection& ConnectionManager, FGatheredReplicationActorLists& OutGatheredList) const;
+
+	/** Add a reference to all the lists held here to the gather list */
+	void AppendAllLists(FGatheredReplicationActorLists& OutGatheredList) const;
+
+	/** Returns the list of all actors held by the permenant and streaming level list. Prefer using AppendAllLists to prevent copying TArrays for nothing. */
+	void GetAllActors(TArray<AActor*>& OutAllActors) const;
+
+	void CountBytes(FArchive& Ar) const;
+
+private:
+
+	FActorRepListRefView PermanentLevelActors;
+
+	FStreamingLevelActorListCollection StreamingLevelActors;
+};
+
 // --------------------------------------------------------------------------------------------------------------------------------------------
 // --------------------------------------------------------------------------------------------------------------------------------------------
 // Per Class/Actor Global/PerConnectino Data Structs
@@ -562,17 +754,30 @@ struct FGlobalActorReplicationInfo
 		ParentActorList.CountBytes(Ar);
 	}
 
+	/** Gather the list of dependent actors filtered by the connection's level visibility */
+	inline void GatherDependentActorLists(const UNetReplicationGraphConnection& ConnectionManager, FGatheredReplicationActorLists& OutGatheredList)
+	{
+		DependentActorList.Gather(ConnectionManager, OutGatheredList);
+	}
+
 	typedef TArray<FActorRepListType> FDependantListType;
-	const FGlobalActorReplicationInfo::FDependantListType& GetDependentActorList() { return DependentActorList; }
+
+	UE_DEPRECATED(5.2, "The dependent actors are exposed via GatherDependentActorLists now.")
+	const FGlobalActorReplicationInfo::FDependantListType& GetDependentActorList()
+	{
+		
+		static FDependantListType EmptyList;
+		return EmptyList;
+	}
 
 	friend struct FGlobalActorReplicationInfoMap;
 
 private:
-	/** When this actor replicates, we replicate these actors immediately afterwards (they are not gathered/prioritized/etc) */
-	FDependantListType DependentActorList;
+	/** When this actor replicates, we replicate these actors immediately afterwards if the client has their level loaded too and the net updated frequency is not too quick */
+	FLevelBasedActorList DependentActorList;
 
 	/** When this actor is added to the dependent list of a parent, track the parent here */
-	FDependantListType ParentActorList;
+	TArray<AActor*> ParentActorList;
 };
 
 /** Templatd struct for mapping UClasses to some data type. The main things this provides is that if a UClass* was not explicitly added, it will climb the class heirachy and find the best match (and then store this for faster lookup next time) */
@@ -765,35 +970,7 @@ struct FGlobalActorReplicationInfoMap
 	}
 
 	/** Removes actor data from map */
-	int32 Remove(const FActorRepListType& RemovedActor)
-	{
-		// Clean the references to the removed actor from its dependency chain.
-		if (FGlobalActorReplicationInfo* RemovedActorInfo = Find(RemovedActor))
-		{
-			// Remove child dependents
-			for (AActor* ChildActor : RemovedActorInfo->DependentActorList)
-			{
-				if (FGlobalActorReplicationInfo* ChildInfo = Find(ChildActor))
-				{
-					ChildInfo->ParentActorList.RemoveSingleSwap(RemovedActor);
-				}
-			}
-
-			// Remove parent dependents
-			for (AActor* ParentActor : RemovedActorInfo->ParentActorList)
-			{
-				if (FGlobalActorReplicationInfo* ParentInfo = Find(ParentActor))
-				{
-					ParentInfo->DependentActorList.RemoveSingleSwap(RemovedActor);
-				}
-			}
-
-			RemovedActorInfo->DependentActorList.Reset();
-			RemovedActorInfo->ParentActorList.Reset();
-		}
-
-		return ActorMap.Remove(RemovedActor);
-	}
+	REPLICATIONGRAPH_API int32 Remove(const FActorRepListType& RemovedActor);
 
 	/** Returns ClassInfo for a given class. */
 	FORCEINLINE FClassReplicationInfo& GetClassInfo(UClass* Class) { return ClassMap.GetChecked(Class); }
@@ -832,15 +1009,30 @@ struct FGlobalActorReplicationInfoMap
 		AllWarnings = WarnAlreadyDependant | WarnParentNotRegistered
 	};
 
+	/**
+	* AddDependentActor makes a dependent actor (the child) considered for replication whenever the parent actor is replicated to a client.
+	* When the parent is replicated, the child actor will be also replicated even if it is outside it's net cull distance.
+	* The child will NOT be replicated if it has no data to send or the net update frequency is too quick since the last netupdate.
+	* Note that this link is one way only, meaning if a child replicates by itself it will not force it's parents to be replicated at the same time.
+	* When setting a dependent link, it is the caller's responsibility to remove the dependency if one of the actor gets destroyed.
+	* 
+	* @param Parent Actor that will try to also replicate the child.
+	* @param Child The actor that we will try to replicate when the parent actor is replicated.
+	* @param WarnFlag Set a flag to trigger an ensure and detect when two actors get set dependent multiple times.
+	*/
 	REPLICATIONGRAPH_API void AddDependentActor(AActor* Parent, AActor* Child, FGlobalActorReplicationInfoMap::EWarnFlag WarnFlag = FGlobalActorReplicationInfoMap::EWarnFlag::None);
 
+	/** 
+	* RemoveDependentActor will remove the given child actor from the parent's dependent list.
+	* This will stop making the child actor try to be replicated whenever the parent actor is
+	*/
 	void RemoveDependentActor(AActor* Parent, AActor* Child)
 	{
 		if (Parent && Child)
 		{
 			if (FGlobalActorReplicationInfo* ParentInfo = Find(Parent))
 			{
-				ParentInfo->DependentActorList.RemoveSingleSwap(Child);
+				ParentInfo->DependentActorList.RemoveNetworkActor(Child);
 			}
 
 			if (FGlobalActorReplicationInfo* ChildInfo = Find(Child))
@@ -850,33 +1042,54 @@ struct FGlobalActorReplicationInfoMap
 		}
 	}
 
-	void RemoveAllActorDependancies(AActor* MainActor)
+	void RemoveAllActorDependencies(AActor* MainActor, FGlobalActorReplicationInfo* MainActorInfo)
 	{
-		FGlobalActorReplicationInfo* MainActorInfo = Find(MainActor);
-		if (MainActor == nullptr)
+		check(MainActorInfo);
+
+		FGatheredReplicationActorLists ListContainer;
+		MainActorInfo->DependentActorList.AppendAllLists(ListContainer);
+
+		// Remove the actor from his child dependents
+		const TArray<FActorRepListConstView>& DependentActorLists = ListContainer.GetLists(EActorRepListTypeFlags::Default);
+		for (const FActorRepListConstView& DependentActorList : DependentActorLists)
 		{
-			return;
+			for (AActor* DependentActor : DependentActorList)
+			{
+				if (FGlobalActorReplicationInfo* ChildInfo = Find(DependentActor))
+				{
+					ChildInfo->ParentActorList.RemoveSingleSwap(MainActor);
+				}
+			}
 		}
-		
-		// Remove this actor from all its parents
-		for (FActorRepListType ParentActor : MainActorInfo->ParentActorList)
+
+		// Remove the actor from his parents
+		FNewReplicatedActorInfo LevelActorInfo(MainActor);
+
+		for (AActor* ParentActor : MainActorInfo->ParentActorList)
 		{
 			if (FGlobalActorReplicationInfo* ParentInfo = Find(ParentActor))
 			{
-				ParentInfo->DependentActorList.RemoveSingleSwap(MainActor);
+				ParentInfo->DependentActorList.RemoveNetworkActor(LevelActorInfo);
 			}
 		}
+
+		MainActorInfo->DependentActorList.Reset();
 		MainActorInfo->ParentActorList.Reset();
+	}
 
-        // Remove all dependant childs from this actor
-		for (FActorRepListType ChildActor : MainActorInfo->DependentActorList)
+	UE_DEPRECATED(5.2, "Fixed a typo. Use RemoveAllActorDependencies instead")
+	void RemoveAllActorDependancies(AActor* MainActor)
+	{
+		RemoveAllActorDependencies(MainActor);
+	}
+
+	void RemoveAllActorDependencies(AActor* MainActor)
+	{
+		FGlobalActorReplicationInfo* MainActorInfo = Find(MainActor);
+		if (MainActor && MainActorInfo)
 		{
-			if (FGlobalActorReplicationInfo* ChildInfo = Find(ChildActor))
-			{
-				ChildInfo->ParentActorList.RemoveSingleSwap(MainActor);
-			}
+			RemoveAllActorDependencies(MainActor, MainActorInfo);
 		}
-		MainActorInfo->DependentActorList.Reset();
 	}
 
 private:
@@ -1192,57 +1405,6 @@ struct FPrioritizedRepList
 	TUniquePtr<TArray<FSkippedActorFullDebugDetails> > SkippedDebugDetails;
 #endif
 };
-
-// --------------------------------------------------------------------------------------------------------------------------------------------
-// --------------------------------------------------------------------------------------------------------------------------------------------
-// Gathering Parameters
-// --------------------------------------------------------------------------------------------------------------------------------------------
-// --------------------------------------------------------------------------------------------------------------------------------------------
-
-// This represents "the list of gathered lists". This is what we push down the Replication Graph and nodes will either Push/Pop List Categories or will add their Replication Lists.
-struct REPLICATIONGRAPH_API FGatheredReplicationActorLists
-{
-	void AddReplicationActorList(const FActorRepListRefView& List, EActorRepListTypeFlags Flags = EActorRepListTypeFlags::Default)
-	{
-#if !(UE_BUILD_SHIPPING || UE_BUILD_TEST)
-		if (CVar_RepGraph_Verify)
-			List.VerifyContents_Slow();
-#endif
-		if (List.Num() > 0)
-		{
-			ReplicationLists[(uint32)Flags].Emplace(FActorRepListConstView(List));
-			CachedNum++;
-		}
-	}
-
-	FORCEINLINE void Reset()
-	{ 
-		for (uint32 i = (uint32)EActorRepListTypeFlags::Default; i < (uint32)EActorRepListTypeFlags::Max; ++i)
-		{
-			ReplicationLists[i].Reset();
-		}
-		CachedNum=0; 
-	}
-	FORCEINLINE int32 NumLists() const 
-	{ 
-		return CachedNum; 
-	}
-	
-	FORCEINLINE const TArray<FActorRepListConstView>& GetLists(EActorRepListTypeFlags ListFlags) const
-	{ 
-		return ReplicationLists[(uint32)ListFlags];
-	}
-	FORCEINLINE bool ContainsLists(EActorRepListTypeFlags Flags) const
-	{ 
-		return ReplicationLists[(uint32)Flags].Num() > 0;
-	}
-	
-private:
-
-	TStaticArray< TArray<FActorRepListConstView>, (uint32)EActorRepListTypeFlags::Max > ReplicationLists;
-	int32 CachedNum = 0;
-};
-
 // --------------------------------------------------------------------------------------------------------------------------------------------
 // --------------------------------------------------------------------------------------------------------------------------------------------
 // Connection Gather Actor List Parameters
@@ -1259,7 +1421,7 @@ typedef TArray<FNetViewer, FReplicationGraphConnectionsAllocator> FNetViewerArra
 // Parameter structure for what we actually pass down during the Gather phase.
 struct FConnectionGatherActorListParameters
 {
-	FConnectionGatherActorListParameters(FNetViewerArray& InViewers, UNetReplicationGraphConnection& InConnectionManager, TSet<FName>& InClientVisibleLevelNamesRef, uint32 InReplicationFrameNum, FGatheredReplicationActorLists& InOutGatheredReplicationLists)
+	FConnectionGatherActorListParameters(FNetViewerArray& InViewers, UNetReplicationGraphConnection& InConnectionManager, const TSet<FName>& InClientVisibleLevelNamesRef, uint32 InReplicationFrameNum, FGatheredReplicationActorLists& InOutGatheredReplicationLists)
 		: Viewers(InViewers), ConnectionManager(InConnectionManager), ReplicationFrameNum(InReplicationFrameNum), OutGatheredReplicationLists(InOutGatheredReplicationLists), ClientVisibleLevelNamesRef(InClientVisibleLevelNamesRef)
 	{
 	}
@@ -1296,29 +1458,6 @@ private:
 	mutable FName LastCheckedVisibleLevelName;
 };
 
-// --------------------------------------------------------------------------------------------------------------------------------------------
-// --------------------------------------------------------------------------------------------------------------------------------------------
-// NewReplicatedActorInfo
-// --------------------------------------------------------------------------------------------------------------------------------------------
-// --------------------------------------------------------------------------------------------------------------------------------------------
-
-/** This is the struct we use to push new replication actors into the graph. "New" doesn't mean "newly spawned" it means "new to the graph". FIXME: Please suggest a better name! */
-struct FNewReplicatedActorInfo
-{
-	explicit FNewReplicatedActorInfo(const FActorRepListType& InActor) : Actor(InActor), Class(InActor->GetClass())
-	{
-		StreamingLevelName = GetStreamingLevelNameOfActor(Actor);
-	}
-
-	AActor* GetActor() const { return Actor; }
-
-	REPLICATIONGRAPH_API static FName GetStreamingLevelNameOfActor(const AActor* Actor);
-
-	FActorRepListType Actor;
-	FName StreamingLevelName;
-	UClass* Class;
-};
-
 // --------------------------------------------------------------------------------------------------------------------------------------------
 // --------------------------------------------------------------------------------------------------------------------------------------------
 // RPCs
@@ -1884,23 +2023,23 @@ public:
 	FActorRepListStatCollector() {}
 
 	/** Collect stats on a single FActorRepList */
-	void VisitRepList(const class UReplicationGraphNode* NodeToVisit, const FActorRepListRefView& RepList);
+	void VisitRepList(const UReplicationGraphNode* NodeToVisit, const FActorRepListRefView& RepList);
 
 	/** Collect stats on a collection of FActorRepLists */
-	void VisitStreamingLevelCollection(const class UReplicationGraphNode* NodeToVisit, const struct FStreamingLevelActorListCollection& StreamingLevelList);
+	void VisitStreamingLevelCollection(const UReplicationGraphNode* NodeToVisit, const struct FStreamingLevelActorListCollection& StreamingLevelList);
 
 	/** Collect stats for FActorRepLists not held by a node */
 	void VisitExplicitStreamingLevelList(FName ListOwnerName, FName StreamLevelName, const FActorRepListRefView& RepList);
 
 	/** Prevents the node from being visited twice in case it's shared multiple times */
-	void FlagNodeVisited(const class UReplicationGraphNode* NodeToVisit);
+	void FlagNodeVisited(const UReplicationGraphNode* NodeToVisit);
 
 	/** Print the statistics previously collected */
 	void PrintCollectedData(FOutputDevice& Ar);
 
 private:
 
-	bool WasNodeVisited(const class UReplicationGraphNode* NodeToVisit);
+	bool WasNodeVisited(const UReplicationGraphNode* NodeToVisit);
 
 private:
 
