Change the BasicReplicationGraph's GridSpatialization2D node to use UE_OLD_WORLD_MAX

diff --git a/Engine/Plugins/Runtime/ReplicationGraph/Source/Private/BasicReplicationGraph.cpp b/Engine/Plugins/Runtime/ReplicationGraph/Source/Private/BasicReplicationGraph.cpp
--- a/Engine/Plugins/Runtime/ReplicationGraph/Source/Private/BasicReplicationGraph.cpp
+++ b/Engine/Plugins/Runtime/ReplicationGraph/Source/Private/BasicReplicationGraph.cpp
@@ -63,7 +63,7 @@ void UBasicReplicationGraph::InitGlobalGraphNodes()
 
 	GridNode = CreateNewNode<UReplicationGraphNode_GridSpatialization2D>();
 	GridNode->CellSize = 10000.f;
-	GridNode->SpatialBias = FVector2D(-WORLD_MAX, -WORLD_MAX);
+	GridNode->SpatialBias = FVector2D(-UE_OLD_WORLD_MAX, -UE_OLD_WORLD_MAX);
 
 	AddGlobalGraphNode(GridNode);
 
