Change the replication graph debug actor to a weak reference so it doesn't break seamless travel when pending kill is disabled

diff --git a/Engine/Plugins/Runtime/ReplicationGraph/Source/Private/ReplicationGraph.cpp b/Engine/Plugins/Runtime/ReplicationGraph/Source/Private/ReplicationGraph.cpp
--- a/Engine/Plugins/Runtime/ReplicationGraph/Source/Private/ReplicationGraph.cpp
+++ b/Engine/Plugins/Runtime/ReplicationGraph/Source/Private/ReplicationGraph.cpp
@@ -1196,11 +1196,11 @@ int32 UReplicationGraph::ServerReplicateActors(float DeltaSeconds)
 #if DO_ENABLE_REPGRAPH_DEBUG_ACTOR
 			{
 				RG_QUICK_SCOPE_CYCLE_COUNTER(NET_ReplicateActors_ReplicateDebugActor);
-				if (ConnectionManager->DebugActor)
+				if (AReplicationGraphDebugActor* DebugActor = ConnectionManager->DebugActor.Get())
 				{
-					FGlobalActorReplicationInfo& GlobalInfo = GlobalActorReplicationInfoMap.Get(ConnectionManager->DebugActor);
-					FConnectionReplicationActorInfo& ActorInfo = ConnectionActorInfoMap.FindOrAdd(ConnectionManager->DebugActor);
-					int64 DebugActorBits = ReplicateSingleActor(ConnectionManager->DebugActor, ActorInfo, GlobalInfo, ConnectionActorInfoMap, *ConnectionManager, FrameNum);
+					FGlobalActorReplicationInfo& GlobalInfo = GlobalActorReplicationInfoMap.Get(DebugActor);
+					FConnectionReplicationActorInfo& ActorInfo = ConnectionActorInfoMap.FindOrAdd(DebugActor);
+					int64 DebugActorBits = ReplicateSingleActor(DebugActor, ActorInfo, GlobalInfo, ConnectionActorInfoMap, *ConnectionManager, FrameNum);
 					// Do not count the debug actor towards our bandwidth limit
 					NetConnection->QueuedBits -= DebugActorBits;
 				}
@@ -2531,7 +2531,7 @@ void UNetReplicationGraphConnection::Serialize(FArchive& Ar)
 void UNetReplicationGraphConnection::TearDown()
 {
 #if DO_ENABLE_REPGRAPH_DEBUG_ACTOR
-	if (DebugActor)
+	if (DebugActor.IsValid())
 	{
 		DebugActor->Destroy();
 	}
@@ -2634,7 +2634,7 @@ void UNetReplicationGraphConnection::InitForConnection(UNetConnection* InConnect
 #if DO_ENABLE_REPGRAPH_DEBUG_ACTOR
 	UReplicationGraph* Graph = Cast<UReplicationGraph>(GetOuter());
 	DebugActor = Graph->CreateDebugActor();
-	if (DebugActor)
+	if (DebugActor.IsValid())
 	{
 		DebugActor->ConnectionManager = this;
 		DebugActor->ReplicationGraph = Graph;
diff --git a/Engine/Plugins/Runtime/ReplicationGraph/Source/Public/ReplicationGraph.h b/Engine/Plugins/Runtime/ReplicationGraph/Source/Public/ReplicationGraph.h
--- a/Engine/Plugins/Runtime/ReplicationGraph/Source/Public/ReplicationGraph.h
+++ b/Engine/Plugins/Runtime/ReplicationGraph/Source/Public/ReplicationGraph.h
@@ -1216,7 +1216,7 @@ public:
 	bool bEnableDebugging;
 
 	UPROPERTY()
-	TObjectPtr<class AReplicationGraphDebugActor> DebugActor = nullptr;
+	TWeakObjectPtr<class AReplicationGraphDebugActor> DebugActor;
 
 	/** Index of the connection in the global list. Will be reassigned when any client disconnects so it is a key that can be referenced only during a single tick */
 	int32 ConnectionOrderNum;
