Send destruction info for dormant actors from cells that are no longer observed by the network client.
- TimeToLife for cell to send destruction is configurable and overridable by the game per GridSpatialization2DNode
- Exposed ReplicatedDormantDestructionInfosPerFrame CVar to have control over optimization levers
- Added Profiler markers to measure impact of enabling dormant actor cleanup

#tests PIE

diff --git a/Engine/Plugins/Runtime/ReplicationGraph/Source/Private/ReplicationGraph.cpp b/Engine/Plugins/Runtime/ReplicationGraph/Source/Private/ReplicationGraph.cpp
--- a/Engine/Plugins/Runtime/ReplicationGraph/Source/Private/ReplicationGraph.cpp
+++ b/Engine/Plugins/Runtime/ReplicationGraph/Source/Private/ReplicationGraph.cpp
@@ -112,6 +112,10 @@ int32 CVar_RepGraph_DormantDynamicActorsDestruction = 0;
 static FAutoConsoleVariableRef CVarRepGraphDormantDynamicActorsDestruction(TEXT("Net.RepGraph.DormantDynamicActorsDestruction"), CVar_RepGraph_DormantDynamicActorsDestruction,
 	TEXT("If true, irrelevant dormant actors will be destroyed on the client"), ECVF_Default);
 
+int32 CVar_RepGraph_DormantCellsTTLDefault = 200;
+static FAutoConsoleVariableRef CVarRepGraphDormantCellsTTLDefault(TEXT("Net.RepGraph.DormantCellsTTLDefault"), CVar_RepGraph_DormantCellsTTLDefault,
+	TEXT("Determine number of frames cell will be considered visible before dormant dynamic actor destruction kicks in"), ECVF_Default);
+
 int32 CVar_RepGraph_ReplicatedDormantDestructionInfosPerFrame = MAX_int32;
 static FAutoConsoleVariableRef CVarRepGraphReplicatedDormantDestructionInfosPerFrame(TEXT("Net.RepGraph.ReplicatedDormantDestructionInfosPerFrame"), CVar_RepGraph_ReplicatedDormantDestructionInfosPerFrame,
 	TEXT("If CVarRepGraphDormantDynamicActorsDestruction is true, this is the max number of destruction infos sent to a client per frame"), ECVF_Default);
@@ -2676,6 +2680,35 @@ bool UNetReplicationGraphConnection::PrepareForReplication()
 
 	BuildVisibleLevels();
 
+	// -------------------------------------------
+	//	Handle dynamic dormant destruction TTL
+	// -------------------------------------------
+	if (!NodesVisibleCells.IsEmpty())
+	{
+		// First pass decreases TTL and removes dead cells
+		for (TTuple<TObjectKey<UReplicationGraphNode>, TArray<FVisibleCellInfo>>& NodeCellInfoPair : NodesVisibleCells)
+		{
+			TArray<FVisibleCellInfo>& CellInfos = NodeCellInfoPair.Value;
+			for (int32 Index = 0; Index < CellInfos.Num(); ++Index)
+			{
+				if (CellInfos[Index].Lifetime < 0)
+				{
+					CellInfos.RemoveAtSwap(Index--);
+				}
+			}
+		}
+
+		// Second pass removes dead nodes
+		for (auto Iterator = NodesVisibleCells.CreateIterator(); Iterator; ++Iterator)
+		{
+			TArray<FVisibleCellInfo>& CellInfos = Iterator.Value();
+			if (CellInfos.IsEmpty())
+			{
+				Iterator.RemoveCurrent();
+			}
+		}
+	}
+
 	return (NetConnection->GetConnectionState() != USOCK_Closed) && (NetConnection->ViewTarget != nullptr) && bConnectionHasCorrectWorld;
 }
 
@@ -2828,6 +2861,11 @@ void UNetReplicationGraphConnection::RemoveActorFromAllPrevDormantActorLists(AAc
 	}
 }
 
+TArray<UNetReplicationGraphConnection::FVisibleCellInfo>& UNetReplicationGraphConnection::GetVisibleCellsForNode(const UReplicationGraphNode* GridNode)
+{
+	return NodesVisibleCells.FindOrAdd(GridNode);
+}
+
 void UNetReplicationGraphConnection::GetClientVisibleLevelNames(TSet<FName>& OutLevelNames) const
 {
 	if (NetConnection == nullptr)
@@ -4859,6 +4897,8 @@ UReplicationGraphNode_GridSpatialization2D::UReplicationGraphNode_GridSpatializa
 {
 	bRequiresPrepareForReplicationCall = true;
 	bDestroyDormantDynamicActors = CVar_RepGraph_GridSpatialization2D_DestroyDormantDynamicActorsDefault != 0;
+	DestroyDormantDynamicActorsCellTTL = CVar_RepGraph_DormantCellsTTLDefault;
+	ReplicatedDormantDestructionInfosPerFrame = CVar_RepGraph_ReplicatedDormantDestructionInfosPerFrame;
 }
 
 void UReplicationGraphNode_GridSpatialization2D::Serialize(FArchive& Ar)
@@ -5679,31 +5719,16 @@ void UReplicationGraphNode_GridSpatialization2D::PrepareForReplication()
 #endif // WITH_SERVER_CODE
 }
 
-// Small structure to make it easier to keep track of 
-// information regarding current players for a connection when working with grids
-struct FPlayerGridCellInformation
-{
-	FPlayerGridCellInformation(FIntPoint InCurLocation) :
-		CurLocation(InCurLocation), PrevLocation(FIntPoint::ZeroValue)
-	{
-	}
-
-	FIntPoint CurLocation;
-	FIntPoint PrevLocation;
-};
-
 void UReplicationGraphNode_GridSpatialization2D::GatherActorListsForConnection(const FConnectionGatherActorListParameters& Params)
 {
 	using namespace UE::Net::Private;
 
 #if WITH_SERVER_CODE
-	TArray<FLastLocationGatherInfo>& LastLocationArray = Params.ConnectionManager.LastGatherLocations;
-	TArray<FVector2D, FReplicationGraphConnectionsAllocator> UniqueCurrentLocations;
+	TArray<FIntPoint, FReplicationGraphConnectionsAllocator> UniqueCurrentGridCells;
 
 	// Consider all users that are in cells for this connection. 
 	// From here, generate a list of coordinates, we'll later work through each coordinate pairing
 	// to find the cells that are actually active. This reduces redundancy and cache misses.
-	TArray<FPlayerGridCellInformation, FReplicationGraphConnectionsAllocator> ActiveGridCells;
 	for (const FNetViewer& CurViewer : Params.Viewers)
 	{
 		if (CurViewer.ViewLocation.Z > ConnectionMaxZ)
@@ -5738,104 +5763,89 @@ void UReplicationGraphNode_GridSpatialization2D::GatherActorListsForConnection(c
 			CellY = 0;
 		}
 
-		FPlayerGridCellInformation NewPlayerCell(FIntPoint(CellX, CellY));
-
-		FLastLocationGatherInfo* GatherInfoForConnection = nullptr;
-
-		// Save this information out for later.
-		if (CurViewer.Connection != nullptr)
-		{
-			GatherInfoForConnection = LastLocationArray.FindByKey<UNetConnection*>(CurViewer.Connection);
-
-			// Add any missing last location information that we don't have
-			if (GatherInfoForConnection == nullptr)
-			{
-				GatherInfoForConnection = &LastLocationArray[LastLocationArray.Emplace(CurViewer.Connection, FVector(ForceInitToZero))];
-			}
-		}
-
-		FVector LastLocationForConnection = GatherInfoForConnection ? GatherInfoForConnection->LastLocation : ClampedViewLoc;
-
-		//@todo: if this is clamp view loc this is now redundant...
-		if (GridBounds.IsValid)
-		{
-			// Clean up the location data for this connection to be grid bound
-			LastLocationForConnection = GridBounds.GetClosestPointTo(LastLocationForConnection);
-		}
-		else
-		{
-			// Prevent extreme locations from causing the Grid to grow too large
-			LastLocationForConnection = LastLocationForConnection.BoundToCube(RepGraphHalfWorldMax);
-		}
-
-		// Try to determine the previous location of the user.
-		NewPlayerCell.PrevLocation.X = FMath::Max(0, (int32)((LastLocationForConnection.X - SpatialBias.X) / CellSize));
-		NewPlayerCell.PrevLocation.Y = FMath::Max(0, (int32)((LastLocationForConnection.Y - SpatialBias.Y) / CellSize));
-
+		FIntPoint CurGridCell {CellX, CellY};
+		
 		// If we have not operated on this cell yet (meaning it's not shared by anyone else), gather for it.
-		if (!UniqueCurrentLocations.Contains(NewPlayerCell.CurLocation))
+		if (!UniqueCurrentGridCells.Contains(CurGridCell))
 		{
-			TArray<UReplicationGraphNode_GridCell*>& GridX = GetGridX(CellX);
-			if (GridX.Num() <= CellY)
+			if (UReplicationGraphNode_GridCell* CellNode = GetCell(GetGridX(CellX), CellY))
 			{
-				GridX.SetNum(CellY + 1);
+				CellNode->GatherActorListsForConnection(Params);
 			}
 
-			UReplicationGraphNode_GridCell* CellNode = GridX[CellY];
-			if (CellNode)
+			UniqueCurrentGridCells.Add(CurGridCell);
+
+			// Reset visible cell TTL
+			if (bDestroyDormantDynamicActors && DestroyDormantDynamicActorsCellTTL > 0)
 			{
-				CellNode->GatherActorListsForConnection(Params);
+				TArray<UNetReplicationGraphConnection::FVisibleCellInfo>& VisibleCells = Params.ConnectionManager.GetVisibleCellsForNode(this);
+				if (UNetReplicationGraphConnection::FVisibleCellInfo* CellInfo = VisibleCells.FindByKey(CurGridCell))
+				{
+					CellInfo->Lifetime = DestroyDormantDynamicActorsCellTTL;
+				}
+				else
+				{
+					VisibleCells.Add(UNetReplicationGraphConnection::FVisibleCellInfo{CurGridCell, DestroyDormantDynamicActorsCellTTL});
+				}
 			}
-
-			UniqueCurrentLocations.Add(NewPlayerCell.CurLocation);
 		}
-
-		// Add this to things we consider later.
-		ActiveGridCells.Add(NewPlayerCell);
 	}
 
 	if (bDestroyDormantDynamicActors && CVar_RepGraph_DormantDynamicActorsDestruction > 0)
 	{
+		QUICK_SCOPE_CYCLE_COUNTER(UReplicationGraphNode_GridSpatialization2D_DestroyDormantDynamicActors);
+
+		FActorRepListRefView DormantActorList; // temp list of "observable dormant actors"
 		FActorRepListRefView& PrevDormantActorList = Params.ConnectionManager.GetPrevDormantActorListForNode(this);
+		TArray<UNetReplicationGraphConnection::FVisibleCellInfo>& VisibleCells = Params.ConnectionManager.GetVisibleCellsForNode(this);
 
-		// Process and create the dormancy list for the active grid for this user
-		for (const FPlayerGridCellInformation& CellInfo : ActiveGridCells)
+		if (!VisibleCells.IsEmpty())
 		{
-			const int32& CellX = CellInfo.CurLocation.X;
-			const int32& CellY = CellInfo.CurLocation.Y;
-			const int32& PrevX = CellInfo.PrevLocation.X;
-			const int32& PrevY = CellInfo.PrevLocation.Y;
-
-			// The idea is that if the previous location is a current location for any other user, we do not bother to do operations on this cell
-			// However, if the current location matches with a current location of another user, continue anyways.
-			//
-			// as above, if the grid cell changed this gather and is not in current use by any other viewer
+			RG_QUICK_SCOPE_CYCLE_COUNTER(UReplicationGraphNode_GridSpatialization2D_CollectDestroyedDormantDynamicActors);
 
-			// TODO: There is a potential list gathering redundancy if two actors share the same current and previous cell information
-			// but this should just result in a wasted cycle if anything.
-			if (((CellX != PrevX) || (CellY != PrevY)) && !UniqueCurrentLocations.Contains(CellInfo.PrevLocation))
+			for (const UNetReplicationGraphConnection::FVisibleCellInfo& CellInfo : VisibleCells)
 			{
-				RG_QUICK_SCOPE_CYCLE_COUNTER(UReplicationGraphNode_GridSpatialization2D_CellChangeDormantRelevancy);
-				FActorRepListRefView DormantActorList;
-
-				TArray<UReplicationGraphNode_GridCell*>& GridX = GetGridX(CellX);
-				UReplicationGraphNode_GridCell* CellNode = GridX[CellY];
+				const FIntPoint GridCell = CellInfo.Location;
+				const int32 Lifetime = CellInfo.Lifetime;
+
+				bool bPassesValidation = true;
+
+				bPassesValidation &= ensureMsgf(
+					(Grid.Num() > GridCell.X && Grid[GridCell.X].Num() > GridCell.Y),
+					TEXT("UReplicationGraphNode_GridSpatialization2D::GatherActorListsForConnection: Previously visible cell (%dx%d) is out of bounds due to grid resize, skipping."),
+					GridCell.X, GridCell.Y);
+				bPassesValidation &= ensureMsgf(
+					Lifetime >= 0,
+					TEXT("UReplicationGraphNode_GridSpatialization2D::GatherActorListsForConnection: Cell (%dx%d) lifetime was negative, that shouldn't happen."),
+					GridCell.X, GridCell.Y);
+				
+				if (!bPassesValidation)
+				{
+					continue;
+				}
 
-				if (CellNode)
+				if (Lifetime == 0)
 				{
-					if (UReplicationGraphNode_DormancyNode* DormancyNode = CellNode->GetDormancyNode())
+					// Add dormant actors from "no longer observable" cell into PrevDormantActorList,
+					// but keep in mind that it could be observable from some other cell, that's why we pass second DormantActorList in
+					if (UReplicationGraphNode_GridCell* PrevCell = GetCell(GetGridX(GridCell.X), GridCell.Y))
 					{
-						// Making sure to remove from the PrevDormantActorList since we don't want things added to the DormantActorList to be destroyed anymore
-						DormancyNode->ConditionalGatherDormantDynamicActors(DormantActorList, Params, nullptr, false, &PrevDormantActorList);
+						if (UReplicationGraphNode_DormancyNode* DormancyNode = PrevCell->GetDormancyNode())
+						{
+							DormancyNode->ConditionalGatherDormantDynamicActors(PrevDormantActorList, Params, &DormantActorList, true);
+						}
 					}
 				}
-
-				// Determine dormant actors for our last location. Do not add actors if they are relevant to anyone.
-				if (UReplicationGraphNode_GridCell* PrevCell = GetCell(GetGridX(PrevX), PrevY))
+				else
 				{
-					if (UReplicationGraphNode_DormancyNode* DormancyNode = PrevCell->GetDormancyNode())
+					// Everything that lives, adds itself to DormantActorList and is excluded from PrevDormantActorList
+					if (UReplicationGraphNode_GridCell* CellNode = GetCell(GetGridX(GridCell.X), GridCell.Y))
 					{
-						DormancyNode->ConditionalGatherDormantDynamicActors(PrevDormantActorList, Params, &DormantActorList, true);
+						if (UReplicationGraphNode_DormancyNode* DormancyNode = CellNode->GetDormancyNode())
+						{
+							// Making sure to remove from the PrevDormantActorList since we don't want things added to the DormantActorList to be destroyed anymore
+							DormancyNode->ConditionalGatherDormantDynamicActors(DormantActorList, Params, nullptr, false, &PrevDormantActorList);
+						}
 					}
 				}
 			}
@@ -5843,7 +5853,9 @@ void UReplicationGraphNode_GridSpatialization2D::GatherActorListsForConnection(c
 
 		if (PrevDormantActorList.Num() > 0)
 		{
-			int32 NumActorsToRemove = CVar_RepGraph_ReplicatedDormantDestructionInfosPerFrame;
+			RG_QUICK_SCOPE_CYCLE_COUNTER(UReplicationGraphNode_GridSpatialization2D_FlushDestroyedDormantDynamicActors);
+
+			int32 NumActorsToRemove = ReplicatedDormantDestructionInfosPerFrame;
 
 			UE_LOG(LogReplicationGraph, Verbose, TEXT("UReplicationGraphNode_GridSpatialization2D::GatherActorListsForConnection: Removing %d Actors (List size: %d)"), FMath::Min(NumActorsToRemove, PrevDormantActorList.Num()), PrevDormantActorList.Num());
 
diff --git a/Engine/Plugins/Runtime/ReplicationGraph/Source/Public/ReplicationGraph.h b/Engine/Plugins/Runtime/ReplicationGraph/Source/Public/ReplicationGraph.h
--- a/Engine/Plugins/Runtime/ReplicationGraph/Source/Public/ReplicationGraph.h
+++ b/Engine/Plugins/Runtime/ReplicationGraph/Source/Public/ReplicationGraph.h
@@ -590,6 +590,10 @@ public:
 
 	// When enabled the RepGraph tells clients to destroy dormant dynamic actors when they go out of relevancy.
 	bool bDestroyDormantDynamicActors = true;
+	// How many frames dormant actors should be outside relevancy range to be destroyed
+	int32 DestroyDormantDynamicActorsCellTTL = 1;
+	// Optimization to spread cost over multiple frames
+	int32 ReplicatedDormantDestructionInfosPerFrame = MAX_int32;
 
 protected:
 
@@ -1267,6 +1271,20 @@ public:
 	FActorRepListRefView& GetPrevDormantActorListForNode(const UReplicationGraphNode* GridNode);
 
 	void RemoveActorFromAllPrevDormantActorLists(AActor* InActor);
+	
+	struct FVisibleCellInfo
+	{
+		FIntPoint Location = {};
+		int32 Lifetime = 0;
+
+		// For FindByKey
+		friend bool operator==(const FVisibleCellInfo& Cell, const FIntPoint& Point)
+		{
+			return Cell.Location == Point;
+		}
+	};
+
+	TArray<FVisibleCellInfo>& GetVisibleCellsForNode(const UReplicationGraphNode* GridNode);
 
 private:
 
@@ -1274,6 +1292,9 @@ private:
 	// the dormant dynamic actor destruction feature.
 	TMap<TObjectKey<UReplicationGraphNode>, FActorRepListRefView> PrevDormantActorListPerNode;
 
+	// History of visible cells per graph node used for defering dormant dynamic actors cleanup
+	TMap<TObjectKey<UReplicationGraphNode>, TArray<FVisibleCellInfo>> NodesVisibleCells;
+	
 	friend UReplicationGraph;
 
 	/** Holds relevant data when parsing deleted actors that could be sent to a viewer */
