Fix TTL decrease for dormant actor cleanup

#rnx

diff --git a/Engine/Plugins/Runtime/ReplicationGraph/Source/Private/ReplicationGraph.cpp b/Engine/Plugins/Runtime/ReplicationGraph/Source/Private/ReplicationGraph.cpp
--- a/Engine/Plugins/Runtime/ReplicationGraph/Source/Private/ReplicationGraph.cpp
+++ b/Engine/Plugins/Runtime/ReplicationGraph/Source/Private/ReplicationGraph.cpp
@@ -2691,7 +2691,7 @@ bool UNetReplicationGraphConnection::PrepareForReplication()
 			TArray<FVisibleCellInfo>& CellInfos = NodeCellInfoPair.Value;
 			for (int32 Index = 0; Index < CellInfos.Num(); ++Index)
 			{
-				if (CellInfos[Index].Lifetime < 0)
+				if (--CellInfos[Index].Lifetime < 0)
 				{
 					CellInfos.RemoveAtSwap(Index--);
 				}
