Fix stale nodes cleanup from caches when Node disconnected from the connection

#rnx

diff --git a/Engine/Plugins/Runtime/ReplicationGraph/Source/Private/ReplicationGraph.cpp b/Engine/Plugins/Runtime/ReplicationGraph/Source/Private/ReplicationGraph.cpp
--- a/Engine/Plugins/Runtime/ReplicationGraph/Source/Private/ReplicationGraph.cpp
+++ b/Engine/Plugins/Runtime/ReplicationGraph/Source/Private/ReplicationGraph.cpp
@@ -2661,6 +2661,8 @@ void UNetReplicationGraphConnection::AddConnectionGraphNode(UReplicationGraphNod
 void UNetReplicationGraphConnection::RemoveConnectionGraphNode(UReplicationGraphNode* Node)
 {
 	ConnectionGraphNodes.RemoveSingleSwap(Node);
+	NodesVisibleCells.Remove(Node);
+	PrevDormantActorListPerNode.Remove(Node);
 }
 
 bool UNetReplicationGraphConnection::PrepareForReplication()
