Fix stale nodes cleanup from caches when Node disconnected from the connection

#rnx

diff --git a/Engine/Plugins/Runtime/ReplicationGraph/Source/Private/ReplicationGraph.cpp b/Engine/Plugins/Runtime/ReplicationGraph/Source/Private/ReplicationGraph.cpp
--- a/Engine/Plugins/Runtime/ReplicationGraph/Source/Private/ReplicationGraph.cpp
+++ b/Engine/Plugins/Runtime/ReplicationGraph/Source/Private/ReplicationGraph.cpp
@@ -2661,8 +2661,7 @@ void UNetReplicationGraphConnection::AddConnectionGraphNode(UReplicationGraphNod
 void UNetReplicationGraphConnection::RemoveConnectionGraphNode(UReplicationGraphNode* Node)
 {
 	ConnectionGraphNodes.RemoveSingleSwap(Node);
-	NodesVisibleCells.Remove(Node);
-	PrevDormantActorListPerNode.Remove(Node);
+	CleanupNodeCaches(Node);
 }
 
 bool UNetReplicationGraphConnection::PrepareForReplication()
@@ -2868,6 +2867,12 @@ TArray<UNetReplicationGraphConnection::FVisibleCellInfo>& UNetReplicationGraphCo
 	return NodesVisibleCells.FindOrAdd(GridNode);
 }
 
+void UNetReplicationGraphConnection::CleanupNodeCaches(const UReplicationGraphNode* Node)
+{
+	NodesVisibleCells.Remove(Node);
+	PrevDormantActorListPerNode.Remove(Node);
+}
+
 void UNetReplicationGraphConnection::GetClientVisibleLevelNames(TSet<FName>& OutLevelNames) const
 {
 	if (NetConnection == nullptr)
diff --git a/Engine/Plugins/Runtime/ReplicationGraph/Source/Public/ReplicationGraph.h b/Engine/Plugins/Runtime/ReplicationGraph/Source/Public/ReplicationGraph.h
--- a/Engine/Plugins/Runtime/ReplicationGraph/Source/Public/ReplicationGraph.h
+++ b/Engine/Plugins/Runtime/ReplicationGraph/Source/Public/ReplicationGraph.h
@@ -1286,6 +1286,8 @@ public:
 
 	TArray<FVisibleCellInfo>& GetVisibleCellsForNode(const UReplicationGraphNode* GridNode);
 
+	void CleanupNodeCaches(const UReplicationGraphNode* Node);
+
 private:
 
 	// Stored list of dormant actors in a previous cell when it's been left - this is for
