ReplicationGraph amortize the cost of dormant actor cleanup by allowing gathering actor lists for 1 connection at a time.

#tests PIE

diff --git a/Engine/Plugins/Runtime/ReplicationGraph/Source/Private/ReplicationGraph.cpp b/Engine/Plugins/Runtime/ReplicationGraph/Source/Private/ReplicationGraph.cpp
--- a/Engine/Plugins/Runtime/ReplicationGraph/Source/Private/ReplicationGraph.cpp
+++ b/Engine/Plugins/Runtime/ReplicationGraph/Source/Private/ReplicationGraph.cpp
@@ -66,6 +66,10 @@ namespace UE::Net::Private
 	// uses a densely-stored grid and risks very high memory usage at large coordinates.
 	constexpr double RepGraphWorldMax = UE_OLD_WORLD_MAX;
 	constexpr double RepGraphHalfWorldMax = UE_OLD_HALF_WORLD_MAX;
+
+	constexpr int32 RepGraphDormantActorsListReservedBuffer = 2048;
+	// TInlineAllocator's size is optimized for default DestroyDormantDynamicActorsCellTTL value
+	constexpr int32 RepGraphDormancyNodesInlineBufferSize = 200;
 }
 
 int32 CVar_RepGraph_Pause = 0;
@@ -132,6 +136,9 @@ static TAutoConsoleVariable<float> CVar_ForceConnectionViewerPriority(TEXT("Net.
 int32 CVar_RepGraph_GridSpatialization2D_DestroyDormantDynamicActorsDefault = 1;
 static FAutoConsoleVariableRef CVarRepGraphGridSpatialization2DDestroyDormantDynamicActorsDefault(TEXT("Net.RepGraph.GridSpatialization2DDestroyDormantDynamicActorsDefault"), CVar_RepGraph_GridSpatialization2D_DestroyDormantDynamicActorsDefault, TEXT("Configure what the default for UReplicationGraphNode_GridSpatialization2D::DestroyDormantDynamicActors should be."), ECVF_Default);
 
+int32 CVar_RepGraph_ConnectionHeavyComputationAmortization = 0;
+static FAutoConsoleVariableRef CVarRepGraphConnectionHeavyComputationAmortization(TEXT("Net.RepGraph.ConnectionHeavyComputationAmortization"), CVar_RepGraph_ConnectionHeavyComputationAmortization, TEXT("Non zero values enable heavy computation cost amortization by updating one connection per frame."), ECVF_Default);
+
 REPGRAPH_DEVCVAR_SHIPCONST(int32, "Net.RepGraph.LogNetDormancyDetails", CVar_RepGraph_LogNetDormancyDetails, 0, "Logs actors that are removed from the replication graph/nodes.");
 REPGRAPH_DEVCVAR_SHIPCONST(int32, "Net.RepGraph.LogActorRemove", CVar_RepGraph_LogActorRemove, 0, "Logs actors that are removed from the replication graph/nodes.");
 REPGRAPH_DEVCVAR_SHIPCONST(int32, "Net.RepGraph.LogActorAdd", CVar_RepGraph_LogActorAdd, 0, "Logs actors that are added to replication graph/nodes.");
@@ -996,6 +1003,9 @@ int32 UReplicationGraph::ServerReplicateActors(float DeltaSeconds)
 	// Total number of children processed, added to all the connections later for stat tracking purposes.
 	int32 NumChildrenConnectionsProcessed = 0;
 
+	// Wrap around connection selector
+	HeavyComputationConnectionSelector = (HeavyComputationConnectionSelector + 1) % Connections.Num();
+
 	for (UNetReplicationGraphConnection* ConnectionManager: Connections)
 	{
 		// Prepare for Replication also handles children as well.
@@ -1074,7 +1084,17 @@ int32 UReplicationGraph::ServerReplicateActors(float DeltaSeconds)
 		
 		FGatheredReplicationActorLists GatheredReplicationListsForConnection;
 
-		const FConnectionGatherActorListParameters Parameters(ConnectionViewers, *ConnectionManager, ConnectionManager->GetCachedClientVisibleLevelNames(), FrameNum, GatheredReplicationListsForConnection);
+		const bool bIsSelectedForHeavyComputation =
+			HeavyComputationConnectionSelector == ConnectionManager->ConnectionOrderNum
+			|| CVar_RepGraph_ConnectionHeavyComputationAmortization == 0;
+
+		const FConnectionGatherActorListParameters Parameters(
+			ConnectionViewers,
+			*ConnectionManager,
+			ConnectionManager->GetCachedClientVisibleLevelNames(),
+			FrameNum,
+			GatheredReplicationListsForConnection,
+			bIsSelectedForHeavyComputation);
 
 		UNetReplicationGraphConnection::FRepGraphDestructionViewerInfoArray DestructionViewersInfo;
 
@@ -2686,16 +2706,12 @@ bool UNetReplicationGraphConnection::PrepareForReplication()
 	// -------------------------------------------
 	if (!NodesVisibleCells.IsEmpty())
 	{
-		// First pass decreases TTL and removes dead cells
+		// First pass decreases TTL
 		for (TTuple<TObjectKey<UReplicationGraphNode>, TArray<FVisibleCellInfo>>& NodeCellInfoPair : NodesVisibleCells)
 		{
-			TArray<FVisibleCellInfo>& CellInfos = NodeCellInfoPair.Value;
-			for (int32 Index = 0; Index < CellInfos.Num(); ++Index)
+			for (FVisibleCellInfo& CellInfo : NodeCellInfoPair.Value)
 			{
-				if (--CellInfos[Index].Lifetime < 0)
-				{
-					CellInfos.RemoveAtSwap(Index--);
-				}
+				CellInfo.Lifetime -= (CellInfo.Lifetime > 0);
 			}
 		}
 
@@ -4717,6 +4733,7 @@ void UReplicationGraphNode_DormancyNode::ConditionalGatherDormantDynamicActors(F
 {
 	auto GatherDormantDynamicActorsForList = [&](const FActorRepListRefView& InReplicationActorList)
 	{
+		RG_QUICK_SCOPE_CYCLE_COUNTER(UReplicationGraphNode_DormancyNode_ConditionalGatherDormantDynamicActors);
 		for (const FActorRepListType& Actor : InReplicationActorList)
 		{
 			if (Actor && !Actor->IsNetStartupActor())
@@ -4731,7 +4748,8 @@ void UReplicationGraphNode_DormancyNode::ConditionalGatherDormantDynamicActors(F
 							continue;
 						}
 
-						if (RemoveFromList && RemoveFromList->RemoveFast(Actor))
+						// Keeping the allocation when removing items let's us safe time on multiple items removed from that list.
+						if (RemoveFromList && RemoveFromList->RemoveFast(Actor, false))
 						{
 							Info->bGridSpatilization_AlreadyDormant = false;
 						}
@@ -4876,9 +4894,9 @@ UReplicationGraphNode* UReplicationGraphNode_GridCell::GetDynamicNode()
 	return DynamicNode;
 }
 
-UReplicationGraphNode_DormancyNode* UReplicationGraphNode_GridCell::GetDormancyNode()
+UReplicationGraphNode_DormancyNode* UReplicationGraphNode_GridCell::GetDormancyNode(bool bInCreateIfMissing)
 {
-	if (DormancyNode == nullptr)
+	if (DormancyNode == nullptr && bInCreateIfMissing)
 	{
 		DormancyNode = CreateChildNode<UReplicationGraphNode_DormancyNode>();
 	}
@@ -5798,11 +5816,10 @@ void UReplicationGraphNode_GridSpatialization2D::GatherActorListsForConnection(c
 		}
 	}
 
-	if (bDestroyDormantDynamicActors && CVar_RepGraph_DormantDynamicActorsDestruction > 0)
+	if (bDestroyDormantDynamicActors && CVar_RepGraph_DormantDynamicActorsDestruction > 0 && Params.bIsSelectedForHeavyComputation)
 	{
 		QUICK_SCOPE_CYCLE_COUNTER(UReplicationGraphNode_GridSpatialization2D_DestroyDormantDynamicActors);
 
-		FActorRepListRefView DormantActorList; // temp list of "observable dormant actors"
 		FActorRepListRefView& PrevDormantActorList = Params.ConnectionManager.GetPrevDormantActorListForNode(this);
 		TArray<UNetReplicationGraphConnection::FVisibleCellInfo>& VisibleCells = Params.ConnectionManager.GetVisibleCellsForNode(this);
 
@@ -5810,11 +5827,15 @@ void UReplicationGraphNode_GridSpatialization2D::GatherActorListsForConnection(c
 		{
 			RG_QUICK_SCOPE_CYCLE_COUNTER(UReplicationGraphNode_GridSpatialization2D_CollectDestroyedDormantDynamicActors);
 
-			for (const UNetReplicationGraphConnection::FVisibleCellInfo& CellInfo : VisibleCells)
+			TArray<UReplicationGraphNode_DormancyNode*, TInlineAllocator<RepGraphDormancyNodesInlineBufferSize>> DormancyNodesCache;
+			DormancyNodesCache.SetNum(VisibleCells.Num());
+
+			// Do validation and cache dormancy nodes
+			for (int32 Index = 0; Index < VisibleCells.Num(); Index++)
 			{
-				const FIntPoint GridCell = CellInfo.Location;
-				const int32 Lifetime = CellInfo.Lifetime;
+				const FIntPoint GridCell = VisibleCells[Index].Location;
 
+#if DO_ENSURE
 				bool bPassesValidation = true;
 
 				bPassesValidation &= ensureMsgf(
@@ -5822,37 +5843,51 @@ void UReplicationGraphNode_GridSpatialization2D::GatherActorListsForConnection(c
 					TEXT("UReplicationGraphNode_GridSpatialization2D::GatherActorListsForConnection: Previously visible cell (%dx%d) is out of bounds due to grid resize, skipping."),
 					GridCell.X, GridCell.Y);
 				bPassesValidation &= ensureMsgf(
-					Lifetime >= 0,
+					VisibleCells[Index].Lifetime >= 0,
 					TEXT("UReplicationGraphNode_GridSpatialization2D::GatherActorListsForConnection: Cell (%dx%d) lifetime was negative, that shouldn't happen."),
 					GridCell.X, GridCell.Y);
-				
+
 				if (!bPassesValidation)
 				{
+					VisibleCells.RemoveAtSwap(Index--, 1, false);
 					continue;
 				}
+#endif
+
+				if (UReplicationGraphNode_GridCell* PrevCell = GetCell(GetGridX(GridCell.X), GridCell.Y))
+				{
+					// don't create a node if missing
+					constexpr bool bDontCreateDormancyNode = false;
+					DormancyNodesCache[Index] = PrevCell->GetDormancyNode(bDontCreateDormancyNode);
+				}
+			}
 
-				if (Lifetime == 0)
+			// Merge Dormant actor lists by producing set C such that
+			// C = A && !B, where
+			// A - set of actors in cells with Lifetime == 0 (i.e. "dead" cells)
+			// B - set of actors in cells with Lifetime > 0 (i.e. "observable" cells)
+			
+			// Set B - GatheredActors, temp list of "observable dormant actors"
+			GatheredActors.Reset(RepGraphDormantActorsListReservedBuffer);
+			for (int32 Index = 0; Index < VisibleCells.Num(); Index++)
+			{
+				if (VisibleCells[Index].Lifetime == 0)
 				{
-					// Add dormant actors from "no longer observable" cell into PrevDormantActorList,
-					// but keep in mind that it could be observable from some other cell, that's why we pass second DormantActorList in
-					if (UReplicationGraphNode_GridCell* PrevCell = GetCell(GetGridX(GridCell.X), GridCell.Y))
+					if (UReplicationGraphNode_DormancyNode* DormancyNode = DormancyNodesCache[Index])
 					{
-						if (UReplicationGraphNode_DormancyNode* DormancyNode = PrevCell->GetDormancyNode())
-						{
-							DormancyNode->ConditionalGatherDormantDynamicActors(PrevDormantActorList, Params, &DormantActorList, true);
-						}
+						// Add dormant actors from "no longer observable" cell into PrevDormantActorList,
+						// but keep in mind that it could be observable from some other cell, that's why we pass second DormantActorList in
+						DormancyNode->ConditionalGatherDormantDynamicActors(PrevDormantActorList, Params, &GatheredActors, true);
 					}
+					// Remove cell after it's processed
+					VisibleCells.RemoveAtSwap(Index--, 1, false);
 				}
 				else
 				{
-					// Everything that lives, adds itself to DormantActorList and is excluded from PrevDormantActorList
-					if (UReplicationGraphNode_GridCell* CellNode = GetCell(GetGridX(GridCell.X), GridCell.Y))
+					if (UReplicationGraphNode_DormancyNode* DormancyNode = DormancyNodesCache[Index])
 					{
-						if (UReplicationGraphNode_DormancyNode* DormancyNode = CellNode->GetDormancyNode())
-						{
-							// Making sure to remove from the PrevDormantActorList since we don't want things added to the DormantActorList to be destroyed anymore
-							DormancyNode->ConditionalGatherDormantDynamicActors(DormantActorList, Params, nullptr, false, &PrevDormantActorList);
-						}
+						// Making sure to remove from the PrevDormantActorList since we don't want things added to the DormantActorList to be destroyed anymore
+						DormancyNode->ConditionalGatherDormantDynamicActors(GatheredActors, Params, nullptr, false, &PrevDormantActorList);
 					}
 				}
 			}
diff --git a/Engine/Plugins/Runtime/ReplicationGraph/Source/Public/ReplicationGraph.h b/Engine/Plugins/Runtime/ReplicationGraph/Source/Public/ReplicationGraph.h
--- a/Engine/Plugins/Runtime/ReplicationGraph/Source/Public/ReplicationGraph.h
+++ b/Engine/Plugins/Runtime/ReplicationGraph/Source/Public/ReplicationGraph.h
@@ -517,7 +517,7 @@ public:
 	// Allow graph to override function for creating the dynamic node in the cell
 	TFunction<UReplicationGraphNode*(UReplicationGraphNode_GridCell* Parent)> CreateDynamicNodeOverride;
 
-	UReplicationGraphNode_DormancyNode* GetDormancyNode();
+	UReplicationGraphNode_DormancyNode* GetDormancyNode(bool bInCreateIfMissing = true);
 
 private:
 
@@ -714,6 +714,9 @@ private:
 	// This is a reused TArray for gathering actor nodes. Just to prevent using a stack based TArray everywhere or static/reset patten.
 	TArray<UReplicationGraphNode_GridCell*> GatheredNodes;
 
+	// This is a reused FActorRepListRefView for gathering actors for dormancy cleanup. Just to prevent using a stack based FActorRepListRefView or static/reset patten.
+	FActorRepListRefView GatheredActors;
+
 	friend class AReplicationGraphDebugActor;
 };
 
@@ -1147,6 +1150,9 @@ private:
 	/** Separate bandwidth cap for traffic used when opening actor channels. Ignored if set to 0 */
 	int32 ActorDiscoveryMaxBitsPerFrame;
 
+	/** Optimization to update only one connection each frame */
+	int32 HeavyComputationConnectionSelector = 0;
+
 	/** Internal time used to track when the next update should occur based on frequency settings. */
 	float TimeLeftUntilUpdate = 0.f;
 
diff --git a/Engine/Plugins/Runtime/ReplicationGraph/Source/Public/ReplicationGraphTypes.h b/Engine/Plugins/Runtime/ReplicationGraph/Source/Public/ReplicationGraphTypes.h
--- a/Engine/Plugins/Runtime/ReplicationGraph/Source/Public/ReplicationGraphTypes.h
+++ b/Engine/Plugins/Runtime/ReplicationGraph/Source/Public/ReplicationGraphTypes.h
@@ -268,9 +268,9 @@ struct REPLICATIONGRAPH_API FActorRepListRefView
 	}
 
 	/** Removes the element quickly but changes the list order */
-	bool RemoveFast(const FActorRepListType& ElementToRemove)
+	bool RemoveFast(const FActorRepListType& ElementToRemove, bool bAllowShrink = true)
 	{
-		return RepList.RemoveSingleSwap(ElementToRemove) > 0;
+		return RepList.RemoveSingleSwap(ElementToRemove, bAllowShrink) > 0;
 	}
 
 	/** Removes the element but keeps the order intact. Generally not recommended for large lists. */
@@ -1417,8 +1417,19 @@ typedef TArray<FNetViewer, FReplicationGraphConnectionsAllocator> FNetViewerArra
 // Parameter structure for what we actually pass down during the Gather phase.
 struct FConnectionGatherActorListParameters
 {
-	FConnectionGatherActorListParameters(FNetViewerArray& InViewers, UNetReplicationGraphConnection& InConnectionManager, const TSet<FName>& InClientVisibleLevelNamesRef, uint32 InReplicationFrameNum, FGatheredReplicationActorLists& InOutGatheredReplicationLists)
-		: Viewers(InViewers), ConnectionManager(InConnectionManager), ReplicationFrameNum(InReplicationFrameNum), OutGatheredReplicationLists(InOutGatheredReplicationLists), ClientVisibleLevelNamesRef(InClientVisibleLevelNamesRef)
+	FConnectionGatherActorListParameters(
+			FNetViewerArray& InViewers,
+			UNetReplicationGraphConnection& InConnectionManager,
+			const TSet<FName>& InClientVisibleLevelNamesRef,
+			uint32 InReplicationFrameNum,
+			FGatheredReplicationActorLists& InOutGatheredReplicationLists,
+			bool bInSelectedForHeavyComputation)
+		: Viewers(InViewers)
+		, ConnectionManager(InConnectionManager)
+		, ReplicationFrameNum(InReplicationFrameNum)
+		, OutGatheredReplicationLists(InOutGatheredReplicationLists)
+		, ClientVisibleLevelNamesRef(InClientVisibleLevelNamesRef)
+		, bIsSelectedForHeavyComputation(bInSelectedForHeavyComputation)
 	{
 	}
 
@@ -1448,6 +1459,7 @@ struct FConnectionGatherActorListParameters
 
 	// Cached off reference for fast Level Visibility lookup
 	const TSet<FName>& ClientVisibleLevelNamesRef;
+	const bool bIsSelectedForHeavyComputation;
 
 private:
 
