Add optional driver name filter to rep graph debug commands

diff --git a/Engine/Plugins/Runtime/ReplicationGraph/Source/Private/ReplicationGraphDebugging.cpp b/Engine/Plugins/Runtime/ReplicationGraph/Source/Private/ReplicationGraphDebugging.cpp
--- a/Engine/Plugins/Runtime/ReplicationGraph/Source/Private/ReplicationGraphDebugging.cpp
+++ b/Engine/Plugins/Runtime/ReplicationGraph/Source/Private/ReplicationGraphDebugging.cpp
@@ -40,12 +40,20 @@
 
 namespace RepGraphDebugging
 {
-	UReplicationGraph* FindReplicationGraphHelper()
+	UReplicationGraph* FindReplicationGraphHelper(const TArray<FString>& Args)
 	{
+		FName NetDriverName = NAME_GameNetDriver;
+
+		const FString* DriverStr = Args.FindByPredicate([](const FString& Str) { return Str.Contains(TEXT("driver=")); });
+		if (DriverStr)
+		{
+			FParse::Value(**DriverStr, TEXT("driver="), NetDriverName);
+		}
+
 		UReplicationGraph* Graph = nullptr;
 		for (TObjectIterator<UReplicationGraph> It; It; ++It)
 		{
-			if (It->NetDriver && It->NetDriver->GetNetMode() != NM_Client)
+			if (It->NetDriver && (NetDriverName == NAME_None || It->NetDriver->NetDriverName == NetDriverName) && It->NetDriver->GetNetMode() != NM_Client)
 			{
 				Graph = *It;
 				break;
@@ -879,16 +887,7 @@ FAutoConsoleCommand RepDriverStarvListCmd(TEXT("Net.RepGraph.StarvedList"), TEXT
 // --------------------------------------------------------------------------------------------------------------------------------------------
 void LogGraphHelper(FOutputDevice& Ar, const TArray< FString >& Args)
 {
-	UReplicationGraph* Graph = nullptr;
-	for (TObjectIterator<UReplicationGraph> It; It; ++It)
-	{
-		if (It->NetDriver && It->NetDriver->GetNetMode() != NM_Client)
-		{
-			Graph = *It;
-			break;
-		}
-	}
-
+	UReplicationGraph* Graph = RepGraphDebugging::FindReplicationGraphHelper(Args);
 	if (!Graph)
 	{
 		return;
@@ -1258,7 +1257,7 @@ TFunction<void()> LogPrioritizedListHelper(FOutputDevice& Ar, const TArray< FStr
 		}
 	};
 
-	UReplicationGraph* Graph = RepGraphDebugging::FindReplicationGraphHelper();
+	UReplicationGraph* Graph = RepGraphDebugging::FindReplicationGraphHelper(Args);
 	if (!Graph)
 	{
 		UE_LOG(LogReplicationGraph, Warning, TEXT("Could not find valid Replication Graph."));
@@ -1350,7 +1349,7 @@ FAutoConsoleCommand RepGraphPrintAllCmd(TEXT("Net.RepGraph.PrintAll"), TEXT(""),
 
 	Args = InArgs;
 
-	UReplicationGraph* Graph = RepGraphDebugging::FindReplicationGraphHelper();
+	UReplicationGraph* Graph = RepGraphDebugging::FindReplicationGraphHelper(Args);
 	if (!Graph)
 	{
 		UE_LOG(LogReplicationGraph, Warning, TEXT("Could not find valid Replication Graph."));
@@ -1405,7 +1404,7 @@ FAutoConsoleCommand RepGraphPrintAllCmd(TEXT("Net.RepGraph.PrintAll"), TEXT(""),
 
 FAutoConsoleCommand RepDriverListsStatsCmd(TEXT("Net.RepGraph.Lists.Stats"), TEXT(""), FConsoleCommandWithArgsDelegate::CreateLambda([](const TArray< FString >& Args)
 {
-	UReplicationGraph* Graph = RepGraphDebugging::FindReplicationGraphHelper();
+	UReplicationGraph* Graph = RepGraphDebugging::FindReplicationGraphHelper(Args);
 	if (!Graph)
 	{
 		UE_LOG(LogReplicationGraph, Warning, TEXT("Could not find valid Replication Graph."));
