Lowering net frequency of rep graph debug actor

diff --git a/Engine/Plugins/Runtime/ReplicationGraph/Source/Private/ReplicationGraph.cpp b/Engine/Plugins/Runtime/ReplicationGraph/Source/Private/ReplicationGraph.cpp
--- a/Engine/Plugins/Runtime/ReplicationGraph/Source/Private/ReplicationGraph.cpp
+++ b/Engine/Plugins/Runtime/ReplicationGraph/Source/Private/ReplicationGraph.cpp
@@ -1224,9 +1224,13 @@ int32 UReplicationGraph::ServerReplicateActors(float DeltaSeconds)
 				{
 					FGlobalActorReplicationInfo& GlobalInfo = GlobalActorReplicationInfoMap.Get(DebugActor);
 					FConnectionReplicationActorInfo& ActorInfo = ConnectionActorInfoMap.FindOrAdd(DebugActor);
-					int64 DebugActorBits = ReplicateSingleActor(DebugActor, ActorInfo, GlobalInfo, ConnectionActorInfoMap, *ConnectionManager, FrameNum);
-					// Do not count the debug actor towards our bandwidth limit
-					NetConnection->QueuedBits -= DebugActorBits;
+
+					if (ReadyForNextReplication(ActorInfo, GlobalInfo, FrameNum))
+					{
+						int64 DebugActorBits = ReplicateSingleActor(DebugActor, ActorInfo, GlobalInfo, ConnectionActorInfoMap, *ConnectionManager, FrameNum);
+						// Do not count the debug actor towards our bandwidth limit
+						NetConnection->QueuedBits -= DebugActorBits;
+					}
 				}
 			}
 #endif
diff --git a/Engine/Plugins/Runtime/ReplicationGraph/Source/Public/ReplicationGraph.h b/Engine/Plugins/Runtime/ReplicationGraph/Source/Public/ReplicationGraph.h
--- a/Engine/Plugins/Runtime/ReplicationGraph/Source/Public/ReplicationGraph.h
+++ b/Engine/Plugins/Runtime/ReplicationGraph/Source/Public/ReplicationGraph.h
@@ -59,7 +59,9 @@ struct FReplicationGraphDestructionSettings;
 
 typedef TObjectKey<class UNetReplicationGraphConnection> FRepGraphConnectionKey;
 
+#ifndef DO_ENABLE_REPGRAPH_DEBUG_ACTOR
 #define DO_ENABLE_REPGRAPH_DEBUG_ACTOR !(UE_BUILD_SHIPPING)
+#endif // DO_ENABLE_REPGRAPH_DEBUG_ACTOR
 
 UCLASS(abstract, transient, config=Engine)
 class REPLICATIONGRAPH_API UReplicationGraphNode : public UObject
@@ -1422,6 +1424,7 @@ public:
 	AReplicationGraphDebugActor()
 	{
 		bReplicates = true; // must be set for RPCs to be sent
+		NetUpdateFrequency = 10.0f;
 	}
 
 	// To prevent demo netdriver from replicating.
