Fix DormantActorDestruction: remove cached nodes, so 2 arrays point to same cell data.

#rnx

#localization none
#tests PIE

diff --git a/Engine/Plugins/Runtime/ReplicationGraph/Source/Private/ReplicationGraph.cpp b/Engine/Plugins/Runtime/ReplicationGraph/Source/Private/ReplicationGraph.cpp
--- a/Engine/Plugins/Runtime/ReplicationGraph/Source/Private/ReplicationGraph.cpp
+++ b/Engine/Plugins/Runtime/ReplicationGraph/Source/Private/ReplicationGraph.cpp
@@ -5853,6 +5853,7 @@ void UReplicationGraphNode_GridSpatialization2D::GatherActorListsForConnection(c
 
 				if (!bPassesValidation)
 				{
+					DormancyNodesCache.RemoveAtSwap(Index, 1, false);
 					VisibleCells.RemoveAtSwap(Index--, 1, false);
 					continue;
 				}
@@ -5884,6 +5885,7 @@ void UReplicationGraphNode_GridSpatialization2D::GatherActorListsForConnection(c
 						DormancyNode->ConditionalGatherDormantDynamicActors(PrevDormantActorList, Params, &GatheredActors, true);
 					}
 					// Remove cell after it's processed
+					DormancyNodesCache.RemoveAtSwap(Index, 1, false);
 					VisibleCells.RemoveAtSwap(Index--, 1, false);
 				}
 				else
