ReplicationGraph
* Prevent sending multicast RPCs on a channel that is closing but not yet cleaned up.

diff --git a/Engine/Plugins/Runtime/ReplicationGraph/Source/Private/ReplicationGraph.cpp b/Engine/Plugins/Runtime/ReplicationGraph/Source/Private/ReplicationGraph.cpp
--- a/Engine/Plugins/Runtime/ReplicationGraph/Source/Private/ReplicationGraph.cpp
+++ b/Engine/Plugins/Runtime/ReplicationGraph/Source/Private/ReplicationGraph.cpp
@@ -2266,6 +2266,12 @@ bool UReplicationGraph::ProcessRemoteFunction(class AActor* Actor, UFunction* Fu
 			{
 				continue;
 			}
+
+			// While the channel is closing we cannot send new multicast RPCs
+			if (ConnectionActorInfo.Channel && ConnectionActorInfo.Channel->Closing)
+			{
+				continue;
+			}
 			
 			//UE_CLOG(ConnectionActorInfo.Channel == nullptr, LogReplicationGraph, Display, TEXT("Null channel on %s for %s"), *GetPathNameSafe(Actor), *GetNameSafe(Function));
 			if (ConnectionActorInfo.Channel == nullptr && (RPC_Multicast_OpenChannelForClass.GetChecked(Actor->GetClass()) == true))
@@ -2611,22 +2617,16 @@ void UNetReplicationGraphConnection::NotifyActorChannelCleanedUp(UActorChannel*
 	{
 		QUICK_SCOPE_CYCLE_COUNTER(UNetReplicationGraphConnection_NotifyActorChannelCleanedUp);
 
-		// No existing way to quickly index from actor channel -> ActorInfo. May want a way to speed this up.
-		// The Actor pointer on the channel would have been set to null previously when the channel was closed,
-		// so we can't use that to look up the actor info by key.
-		// Also, the actor may be destroyed and garbage collected before this point.
-
 		FConnectionReplicationActorInfo* ActorInfo = ActorInfoMap.FindByChannel(Channel);
 		if (ActorInfo)
 		{
-			// Note we can't directly remove the entry from ActorInfoMap.ActorMap since we don't have the AActor* to key into that map
-			// But we don't actually have to remove the entry since we no longer iterate through ActorInfoMap.ActorMap in non debug functions.
-			// So all we need to do is clear the runtime/transient data for this actorinfo map. (We want to preserve the dormancy flag and the 
-			// settings we pulled from the FGlobalActorReplicationInfo, but clear the frame counters, etc).
-
+			// Note we don't remove the entry from ActorInfoMap.ActorMap here since the channel does not hold the Actor anymore.
+			// The ActorMap entry will be removed by RemoveNetworkActor instead.
+			
 			if (Channel == ActorInfo->Channel)
 			{
 				// Only reset our state if we're still the associated channel.
+				// We want to preserve the dormancy flag and the settings we pulled from the FGlobalActorReplicationInfo, but clear the frame counters, etc
 				ActorInfo->ResetFrameCounters();
 			}
 			else
@@ -2635,10 +2635,7 @@ void UNetReplicationGraphConnection::NotifyActorChannelCleanedUp(UActorChannel*
 			}
 
 			// Remove reference from channel map
-			// We call this last, as it could be the last thing holding onto the underlying
-			// shared pointer and we don't want to try and access potentially garbage memory.
-			// This isn't a big deal for now since FConnectionReplicationActorInfo is just a POD
-			// type, but if that changes it could be a problem.
+			// ActorInfo should not be accessed after this since it could be destroyed.
 			ActorInfoMap.RemoveChannel(Channel);
 		}
 	}
