Enabling UnsafeTypeCastWarningLevel for rep graph

diff --git a/Engine/Plugins/Runtime/ReplicationGraph/Source/Private/ReplicationGraph.cpp b/Engine/Plugins/Runtime/ReplicationGraph/Source/Private/ReplicationGraph.cpp
--- a/Engine/Plugins/Runtime/ReplicationGraph/Source/Private/ReplicationGraph.cpp
+++ b/Engine/Plugins/Runtime/ReplicationGraph/Source/Private/ReplicationGraph.cpp
@@ -54,6 +54,7 @@
 #include "Net/Core/Trace/NetTrace.h"
 #include "Engine/ServerStatReplicator.h"
 #include "Stats/StatsTrace.h"
+#include <limits>
 
 #include UE_INLINE_GENERATED_CPP_BY_NAME(ReplicationGraph)
 
@@ -1227,7 +1228,7 @@ int32 UReplicationGraph::ServerReplicateActors(float DeltaSeconds)
 
 					if (ReadyForNextReplication(ActorInfo, GlobalInfo, FrameNum))
 					{
-						int64 DebugActorBits = ReplicateSingleActor(DebugActor, ActorInfo, GlobalInfo, ConnectionActorInfoMap, *ConnectionManager, FrameNum);
+						const int32 DebugActorBits = IntCastChecked<int32>(ReplicateSingleActor(DebugActor, ActorInfo, GlobalInfo, ConnectionActorInfoMap, *ConnectionManager, FrameNum));
 						// Do not count the debug actor towards our bandwidth limit
 						NetConnection->QueuedBits -= DebugActorBits;
 					}
@@ -1372,13 +1373,13 @@ void UReplicationGraph::ReplicateActorListsForConnections_Default(UNetReplicatio
 				// -------------------
 				if (GlobalData.Settings.DistancePriorityScale > 0.f)
 				{
-					float SmallestDistanceSq = TNumericLimits<float>::Max();
+					FVector::FReal SmallestDistanceSq = std::numeric_limits<FVector::FReal>::max();
 					int32 ViewersThatSkipActor = 0;
 					
 					for (const FNetViewer& CurViewer : Viewers)
 					{
-						const float DistSq = (GlobalData.WorldLocation - CurViewer.ViewLocation).SizeSquared();
-						SmallestDistanceSq = FMath::Min<float>(DistSq, SmallestDistanceSq);
+						const FVector::FReal DistSq = (GlobalData.WorldLocation - CurViewer.ViewLocation).SizeSquared();
+						SmallestDistanceSq = FMath::Min(DistSq, SmallestDistanceSq);
 
 						// Figure out if we should be skipping this actor
 						if (bDoDistanceCull && ConnectionData.GetCullDistanceSquared() > 0.f && DistSq > ConnectionData.GetCullDistanceSquared())
@@ -1391,7 +1392,7 @@ void UReplicationGraph::ReplicateActorListsForConnections_Default(UNetReplicatio
 					// If no one is near this actor, skip it.
 					if (ViewersThatSkipActor >= Viewers.Num())
 					{
-						DO_REPGRAPH_DETAILS(PrioritizedReplicationList.GetNextSkippedDebugDetails(Actor)->DistanceCulled = FMath::Sqrt(SmallestDistanceSq));
+						DO_REPGRAPH_DETAILS(PrioritizedReplicationList.GetNextSkippedDebugDetails(Actor)->DistanceCulled = FloatCastChecked<float>(FMath::Sqrt(SmallestDistanceSq), UE::LWC::DefaultFloatPrecision));
 
 						// Skipped actors should not have any 
 						if (bDoCulledOnConnectionCount)
@@ -1401,7 +1402,7 @@ void UReplicationGraph::ReplicateActorListsForConnections_Default(UNetReplicatio
 						continue;
 					}
 
-					const float DistanceFactor = FMath::Clamp<float>((SmallestDistanceSq) / MaxDistanceScaling, 0.f, 1.f) * GlobalData.Settings.DistancePriorityScale;
+					const float DistanceFactor = FMath::Clamp(FloatCastChecked<float>(SmallestDistanceSq / MaxDistanceScaling, UE::LWC::DefaultFloatPrecision), 0.f, 1.f) * GlobalData.Settings.DistancePriorityScale;
 					if (DO_REPGRAPH_DETAILS(UNLIKELY(DebugDetails)))
 					{
 						DebugDetails->DistanceSq = SmallestDistanceSq;
@@ -1431,7 +1432,7 @@ void UReplicationGraph::ReplicateActorListsForConnections_Default(UNetReplicatio
 
 					if (DO_REPGRAPH_DETAILS(UNLIKELY(DebugDetails)))
 					{
-						DebugDetails->FramesSinceLastRap = FramesSinceLastRep;
+						DebugDetails->FramesSinceLastRap = static_cast<uint32>(FMath::TruncToInt32(FramesSinceLastRep));
 						DebugDetails->StarvationFactor = StarvationFactor;
 					}
 				}
@@ -1691,7 +1692,7 @@ void UReplicationGraph::ReplicateActorListsForConnections_FastShared(UNetReplica
 				}
 
 				// Simple distance cull
-				const float DistSq = DirToActor.SizeSquared();
+				const FVector::FReal DistSq = DirToActor.SizeSquared();
 				if (!(DistSq > (ConnectionData.GetCullDistanceSquared() * FastSharedDistanceRequirementPct)))
 				{
 					bNoViewRelevency = false;
@@ -2054,10 +2055,12 @@ int64 UReplicationGraph::ReplicateSingleActor(AActor* Actor, FConnectionReplicat
 	// Optional budget for actor discovery traffic
 	if (bIsTrafficActorDiscovery && !bIsActorDiscoveryBudgetFull)
 	{
-		ConnectionManager.QueuedBitsForActorDiscovery += BitsWritten;
+		const int32 DiscoveryBits = IntCastChecked<int32>(BitsWritten);
+
+		ConnectionManager.QueuedBitsForActorDiscovery += DiscoveryBits;
 
 		// Remove the discovery traffic from the regular traffic
-		NetConnection->QueuedBits -= BitsWritten;
+		NetConnection->QueuedBits -= DiscoveryBits;
 		BitsWritten = 0;
 	}
 
@@ -2307,7 +2310,7 @@ bool UReplicationGraph::ProcessRemoteFunction(class AActor* Actor, UFunction* Fu
 						// default to the channel being closed.
 						for (const FNetViewer& Viewer : ViewsToConsider)
 						{
-							const float DistSq = (ActorLocation.GetValue() - Viewer.ViewLocation).SizeSquared();
+							const FVector::FReal DistSq = (ActorLocation.GetValue() - Viewer.ViewLocation).SizeSquared();
 							if (DistSq <= ConnectionActorInfo.GetCullDistanceSquared())
 							{
 								bShouldOpenChannel = true;
@@ -2948,7 +2951,7 @@ int64 UNetReplicationGraphConnection::ReplicateDestructionInfos(const FRepGraphD
 			// Only send destruction info if the viewers are close enough to the destroyed actor
 			for (const FRepGraphDestructionViewerInfo& Viewer : DestructionViewersInfo)
 			{
-				const float DistSquared = FVector::DistSquared2D(Info.CachedPosition, Viewer.ViewerLocation);
+				const FVector::FReal DistSquared = FVector::DistSquared2D(Info.CachedPosition, Viewer.ViewerLocation);
 
 				if (DistSquared < DestructionSettings.DestructInfoMaxDistanceSquared)
 				{
@@ -2956,7 +2959,7 @@ int64 UNetReplicationGraphConnection::ReplicateDestructionInfos(const FRepGraphD
 					break;
 				}
 
-				const float OutOfRangeDistSquared = FVector::DistSquared2D(Info.CachedPosition, Viewer.LastOutOfRangeLocationCheck);
+				const FVector::FReal OutOfRangeDistSquared = FVector::DistSquared2D(Info.CachedPosition, Viewer.LastOutOfRangeLocationCheck);
 				
 				// Add the actor to the OutOfRangeList only if it is outside the range from the next check. If not, keep it in the destruction list to be evaluated next frame.
 				if (OutOfRangeDistSquared <= DestructionSettings.MaxPendingListDistanceSquared)
@@ -3048,7 +3051,7 @@ void UNetReplicationGraphConnection::OnUpdateViewerLocation(FLastLocationGatherI
 {
 	const bool bIgnoreDistanceCheck = DestructionSettings.OutOfRangeDistanceCheckThresholdSquared == 0.0f;
 
-	const float OutOfRangeDistanceSquared = FVector::DistSquared2D(Viewer.ViewLocation, LocationInfo->LastOutOfRangeLocationCheck);
+	const FVector::FReal OutOfRangeDistanceSquared = FVector::DistSquared2D(Viewer.ViewLocation, LocationInfo->LastOutOfRangeLocationCheck);
 
 	// Test all accumulated out of range actors only once the viewer has gone far enough from the last check
 	if( bIgnoreDistanceCheck || (OutOfRangeDistanceSquared > DestructionSettings.OutOfRangeDistanceCheckThresholdSquared) )
@@ -3061,7 +3064,7 @@ void UNetReplicationGraphConnection::OnUpdateViewerLocation(FLastLocationGatherI
 		{
 			FCachedDestructInfo& CachedInfo = OutOfRangeDestroyedActors[Index];
 
-			const float ActorDistSquared = FVector::DistSquared2D(CachedInfo.CachedPosition, Viewer.ViewLocation);
+			const FVector::FReal ActorDistSquared = FVector::DistSquared2D(CachedInfo.CachedPosition, Viewer.ViewLocation);
 				
 			if (ActorDistSquared <= DestructionSettings.MaxPendingListDistanceSquared)
 			{
@@ -3754,9 +3757,9 @@ void UReplicationGraphNode_ActorListFrequencyBuckets::TearDown()
 
 namespace RepGraphDynamicSpatialFrequency
 {
-	const float AssumedTickRate = 30.f; // UNetDriver::NetServerMaxTickRate
-	const float TargetKBytesSec = 10.f; // 10K/sec
-	const int64 BitsPerFrame = TargetKBytesSec * 1024.f * 8.f / AssumedTickRate;
+	constexpr float AssumedTickRate = 30.f; // UNetDriver::NetServerMaxTickRate
+	constexpr float TargetKBytesSec = 10.f; // 10K/sec
+	constexpr int64 BitsPerFrame = (int64)((TargetKBytesSec * 1024.f * 8.f) / AssumedTickRate);
 };
 
 static TArray<UReplicationGraphNode_DynamicSpatialFrequency::FSpatializationZone>& DefaultSpatializationZones()
@@ -3974,7 +3977,7 @@ void UReplicationGraphNode_DynamicSpatialFrequency::GatherActorListsForConnectio
 			
 			else if (Item.EnableFastPath && ReadyForNextReplication_FastPath(ConnectionInfo, GlobalInfo, FrameNum))
 			{
-				const int64 FastSharedBits = RepGraph->ReplicateSingleActor_FastShared(Actor, ConnectionInfo, GlobalInfo, Params.ConnectionManager, FrameNum);
+				const int32 FastSharedBits = IntCastChecked<int32>(RepGraph->ReplicateSingleActor_FastShared(Actor, ConnectionInfo, GlobalInfo, Params.ConnectionManager, FrameNum));
 				QueuedBits -= FastSharedBits; // We are doing our own bandwidth limiting here, so offset the netconnection's tracking.
 				BitsWritten += FastSharedBits;
 			}
@@ -3987,10 +3990,10 @@ void UReplicationGraphNode_DynamicSpatialFrequency::GatherActorListsForConnectio
 			}
 		}
 
-
 		if (CVar_RepGraph_DynamicSpatialFrequency_UncapBandwidth > 0)
 		{
-			UE_LOG(LogReplicationGraph, Display, TEXT("Uncapped bandwidth usage of UReplicationGraphNode_DynamicSpatialFrequency = %d bits -> %d bytes -> %.2f KBytes/sec"), BitsWritten, (BitsWritten+7) >> 3,  ((float)((BitsWritten+7)>>3)/1024.f) * GraphGlobals->ReplicationGraph->NetDriver->NetServerMaxTickRate);
+			UE_LOG(LogReplicationGraph, Display, TEXT("Uncapped bandwidth usage of UReplicationGraphNode_DynamicSpatialFrequency = %d bits -> %d bytes -> %.2f KBytes/sec"),
+				BitsWritten, (BitsWritten + 7) >> 3, (static_cast<double>((BitsWritten + 7) >> 3) / 1024.0) * GraphGlobals->ReplicationGraph->NetDriver->NetServerMaxTickRate);
 		}
 	}
 #endif // WITH_SERVER_CODE
@@ -4079,13 +4082,13 @@ void UReplicationGraphNode_DynamicSpatialFrequency::CalcFrequencyForActor(AActor
 		return;
 	}
 
-	float SmallestDistanceToActorSq = TNumericLimits<float>::Max();
+	FVector::FReal SmallestDistanceToActorSq = std::numeric_limits<FVector::FReal>::max();
 	const FNetViewer* LowestDistanceViewer = nullptr;
 
 	// Find the closest viewer to this item or the first viewer if there are no viewers closer.
 	for (const FNetViewer& CurViewer : Viewers)
 	{
-		float CurDistance = (GlobalInfo.WorldLocation - CurViewer.ViewLocation).SizeSquared();
+		FVector::FReal CurDistance = (GlobalInfo.WorldLocation - CurViewer.ViewLocation).SizeSquared();
 		if (LowestDistanceViewer == nullptr || CurDistance < SmallestDistanceToActorSq)
 		{
 			LowestDistanceViewer = &CurViewer;
@@ -4108,9 +4111,9 @@ void UReplicationGraphNode_DynamicSpatialFrequency::CalcFrequencyForActor(AActor
 	// --------------------------------------------------------------------------------------------------------
 	const FVector& ConnectionViewDir = LowestDistanceViewer->ViewDir;
 	const FVector DirToActor = (GlobalInfo.WorldLocation - LowestDistanceViewer->ViewLocation);
-	const float DistanceToActor = FMath::Sqrt(SmallestDistanceToActorSq);
+	const FVector::FReal DistanceToActor = FMath::Sqrt(SmallestDistanceToActorSq);
 	const FVector NormDirToActor = DistanceToActor > SMALL_NUMBER ? (DirToActor / DistanceToActor) : DirToActor;
-	const float DotP = FMath::Clamp(FVector::DotProduct(NormDirToActor, ConnectionViewDir), -1.0f, 1.0f);
+	const float DotP = FMath::Clamp(FloatCastChecked<float>(FVector::DotProduct(NormDirToActor, ConnectionViewDir), UE::LWC::DefaultFloatPrecision), -1.0f, 1.0f);
 
 	const bool ActorSupportsFastShared = (GlobalInfo.Settings.FastSharedReplicationFunc != nullptr);
 	TArrayView<FSpatializationZone>& ZoneList = ActorSupportsFastShared ? MySettings.ZoneSettings : MySettings.ZoneSettings_NonFastSharedActors;
@@ -4142,7 +4145,7 @@ void UReplicationGraphNode_DynamicSpatialFrequency::CalcFrequencyForActor(AActor
 				}
 
 				const float CullDist = ConnectionInfo.GetCullDistance();
-				const float DistPct =  DistanceToActor / CullDist;
+				const float DistPct = FloatCastChecked<float>(DistanceToActor / CullDist, UE::LWC::DefaultFloatPrecision);
 
 				const float BiasDistPct = DistPct - ZoneInfo.MinDistPct;
 				const float FinalPCT = FMath::Clamp<float>( BiasDistPct / (ZoneInfo.MaxDistPct - ZoneInfo.MinDistPct), 0.f, 1.f);
@@ -4286,7 +4289,7 @@ void UReplicationGraphNode_DynamicSpatialFrequency::GatherActors_DistanceOnly(co
 	for (AActor* Actor : RepList)
 	{
 		FGlobalActorReplicationInfo& GlobalInfo = GlobalMap.Get(Actor);
-		float ShortestDistanceToActorSq = TNumericLimits<float>::Max();
+		FVector::FReal ShortestDistanceToActorSq = std::numeric_limits<FVector::FReal>::max();
 		bool bShouldSkipActor = false;
 
 		// Don't replicate the connection view target like this. It will be done through a connection specific node
@@ -4298,7 +4301,7 @@ void UReplicationGraphNode_DynamicSpatialFrequency::GatherActors_DistanceOnly(co
 				break;
 			}
 
-			ShortestDistanceToActorSq = FMath::Min<float>((GlobalInfo.WorldLocation - CurViewer.ViewLocation).SizeSquared(), ShortestDistanceToActorSq);
+			ShortestDistanceToActorSq = FMath::Min((GlobalInfo.WorldLocation - CurViewer.ViewLocation).SizeSquared(), ShortestDistanceToActorSq);
 		}
 
 		if (bShouldSkipActor)
@@ -4306,7 +4309,7 @@ void UReplicationGraphNode_DynamicSpatialFrequency::GatherActors_DistanceOnly(co
 			continue;
 		}
 
-		SortedReplicationList.Emplace(Actor, (int32)ShortestDistanceToActorSq, &GlobalInfo);
+		SortedReplicationList.Emplace(Actor, FMath::TruncToInt32(ShortestDistanceToActorSq), &GlobalInfo);
 	}
 }
 
@@ -5228,14 +5231,14 @@ UReplicationGraphNode_GridSpatialization2D::FActorCellInfo UReplicationGraphNode
 	}
 
 	FActorCellInfo CellInfo;
-	const auto LocationBiasX = (ClampedLocation.X - SpatialBias.X);
-	const auto LocationBiasY = (ClampedLocation.Y - SpatialBias.Y);
+	const FVector::FReal LocationBiasX = (ClampedLocation.X - SpatialBias.X);
+	const FVector::FReal LocationBiasY = (ClampedLocation.Y - SpatialBias.Y);
 
-	const auto Dist = CullDistance;
-	const auto MinX = LocationBiasX - Dist;
-	const auto MinY = LocationBiasY - Dist;
-	auto MaxX = LocationBiasX + Dist;
-	auto MaxY = LocationBiasY + Dist;
+	const FVector::FReal Dist = CullDistance;
+	const FVector::FReal MinX = LocationBiasX - Dist;
+	const FVector::FReal MinY = LocationBiasY - Dist;
+	FVector::FReal MaxX = LocationBiasX + Dist;
+	FVector::FReal MaxY = LocationBiasY + Dist;
 
 	if (GridBounds.IsValid)
 	{
@@ -5244,11 +5247,11 @@ UReplicationGraphNode_GridSpatialization2D::FActorCellInfo UReplicationGraphNode
 		MaxY = FMath::Min(MaxY, BoundSize.Y);
 	}
 
-	CellInfo.StartX = FMath::Max<int32>(0, MinX / CellSize);
-	CellInfo.StartY = FMath::Max<int32>(0, MinY / CellSize);
+	CellInfo.StartX = FMath::Max<int32>(0, UE::LWC::FloatToIntCastChecked<int32>(MinX / CellSize));
+	CellInfo.StartY = FMath::Max<int32>(0, UE::LWC::FloatToIntCastChecked<int32>(MinY / CellSize));
 
-	CellInfo.EndX = FMath::Max<int32>(0, MaxX / CellSize);
-	CellInfo.EndY = FMath::Max<int32>(0, MaxY / CellSize);
+	CellInfo.EndX = FMath::Max<int32>(0, UE::LWC::FloatToIntCastChecked<int32>(MaxX / CellSize));
+	CellInfo.EndY = FMath::Max<int32>(0, UE::LWC::FloatToIntCastChecked<int32>(MaxY / CellSize));
 	return CellInfo;
 }
 
@@ -5337,7 +5340,7 @@ static FAutoConsoleVariableRef CVarRepSpatialPauseDynamic(TEXT("Net.RepGraph.Spa
 int32 CVar_RepGraph_Spatial_DebugDynamic = 0;
 static FAutoConsoleVariableRef CVarRepGraphSpatialDebugDynamic(TEXT("Net.RepGraph.Spatial.DebugDynamic"), CVar_RepGraph_Spatial_DebugDynamic, TEXT("Prints debug info whenever dynamic actors changes spatial cells"), ECVF_Default );
 
-int32 CVar_RepGraph_Spatial_BiasCreep = 0.f;
+float CVar_RepGraph_Spatial_BiasCreep = 0.f;
 static FAutoConsoleVariableRef CVarRepGraphSpatialBiasCreep(TEXT("Net.RepGraph.Spatial.BiasCreep"), CVar_RepGraph_Spatial_BiasCreep, TEXT("Changes bias each frame by this much and force rebuld. For stress test debugging"), ECVF_Default );
 
 void UReplicationGraphNode_GridSpatialization2D::PrepareForReplication()
@@ -5775,14 +5778,14 @@ void UReplicationGraphNode_GridSpatialization2D::GatherActorListsForConnection(c
 		}
 
 		// Find out what bucket the view is in
-		int32 CellX = (ClampedViewLoc.X - SpatialBias.X) / CellSize;
+		int32 CellX = UE::LWC::FloatToIntCastChecked<int32>((ClampedViewLoc.X - SpatialBias.X) / CellSize);
 		if (CellX < 0)
 		{
 			UE_LOG(LogReplicationGraph, Verbose, TEXT("Net view location.X %s is less than the spatial bias %s for %s"), *ClampedViewLoc.ToString(), *SpatialBias.ToString(), CurViewer.Connection ? *CurViewer.Connection->Describe() : TEXT("NONE"));
 			CellX = 0;
 		}
 
-		int32 CellY = (ClampedViewLoc.Y - SpatialBias.Y) / CellSize;
+		int32 CellY = UE::LWC::FloatToIntCastChecked<int32>((ClampedViewLoc.Y - SpatialBias.Y) / CellSize);
 		if (CellY < 0)
 		{
 			UE_LOG(LogReplicationGraph, Verbose, TEXT("Net view location.Y %s is less than the spatial bias %s for %s"), *ClampedViewLoc.ToString(), *SpatialBias.ToString(), CurViewer.Connection ? *CurViewer.Connection->Describe() : TEXT("NONE"));
diff --git a/Engine/Plugins/Runtime/ReplicationGraph/Source/Private/ReplicationGraphDebugging.cpp b/Engine/Plugins/Runtime/ReplicationGraph/Source/Private/ReplicationGraphDebugging.cpp
--- a/Engine/Plugins/Runtime/ReplicationGraph/Source/Private/ReplicationGraphDebugging.cpp
+++ b/Engine/Plugins/Runtime/ReplicationGraph/Source/Private/ReplicationGraphDebugging.cpp
@@ -412,8 +412,8 @@ void AReplicationGraphDebugActor::ServerCellInfo_Implementation()
 			return;
 		}
 
-		int32 CellX = FMath::Max<int32>(0, (Viewer.ViewLocation.X - GridNode->SpatialBias.X) / GridNode->CellSize);
-		int32 CellY = FMath::Max<int32>(0, (Viewer.ViewLocation.Y - GridNode->SpatialBias.Y) / GridNode->CellSize);
+		const int32 CellX = FMath::Max<int32>(0, UE::LWC::FloatToIntCastChecked<int32>((Viewer.ViewLocation.X - GridNode->SpatialBias.X) / GridNode->CellSize));
+		const int32 CellY = FMath::Max<int32>(0, UE::LWC::FloatToIntCastChecked<int32>((Viewer.ViewLocation.Y - GridNode->SpatialBias.Y) / GridNode->CellSize)); 
 
 		FVector CellLocation(GridNode->SpatialBias.X + (((float)(CellX)+0.5f) * GridNode->CellSize), GridNode->SpatialBias.Y + (((float)(CellY)+0.5f) * GridNode->CellSize), Viewer.ViewLocation.Z);
 
@@ -547,8 +547,14 @@ void AReplicationGraphDebugActor::ServerSetPeriodFrameForClass_Implementation(UC
 		return;
 	}
 
+	if (PeriodFrame < 0 || PeriodFrame > MAX_uint16)
+	{
+		UE_LOG(LogReplicationGraph, Display, TEXT("Invalid PeriodFrame"));
+		return;
+	}
+
 	FClassReplicationInfo& ClassInfo = ReplicationGraph->GlobalActorReplicationInfoMap.GetClassInfo(Class);
-	ClassInfo.ReplicationPeriodFrame = PeriodFrame;
+	ClassInfo.ReplicationPeriodFrame = static_cast<uint16>(PeriodFrame);
 	UE_LOG(LogReplicationGraph, Display, TEXT("Setting ReplicationPeriodFrame for class %s to %u"), *Class->GetName(), ClassInfo.ReplicationPeriodFrame);
 
 	for (TActorIterator<AActor> ActorIt(GetWorld(), Class); ActorIt; ++ActorIt)
@@ -556,14 +562,14 @@ void AReplicationGraphDebugActor::ServerSetPeriodFrameForClass_Implementation(UC
 		AActor* Actor = *ActorIt;
 		if (FGlobalActorReplicationInfo* ActorInfo = ReplicationGraph->GlobalActorReplicationInfoMap.Find(Actor))
 		{
-			ActorInfo->Settings.ReplicationPeriodFrame = PeriodFrame;
+			ActorInfo->Settings.ReplicationPeriodFrame = static_cast<uint16>(PeriodFrame);
 			UE_LOG(LogReplicationGraph, Display, TEXT("Setting GlobalActorInfo ReplicationPeriodFrame for %s to %u"), *Actor->GetName(), ActorInfo->Settings.ReplicationPeriodFrame);
 		}
 
 
 		if (FConnectionReplicationActorInfo* ConnectionActorInfo = ConnectionManager->ActorInfoMap.Find(Actor))
 		{
-			ConnectionActorInfo->ReplicationPeriodFrame = PeriodFrame;
+			ConnectionActorInfo->ReplicationPeriodFrame = static_cast<uint16>(PeriodFrame);
 			UE_LOG(LogReplicationGraph, Display, TEXT("Setting Connection ReplicationPeriodFrame for %s to %u"), *Actor->GetName(), ConnectionActorInfo->ReplicationPeriodFrame);
 		}
 	}
@@ -593,7 +599,7 @@ FAutoConsoleCommandWithWorldAndArgs NetRepGraphSetPeriodFrame(TEXT("Net.RepGraph
 
 		for (TActorIterator<AReplicationGraphDebugActor> It(World); It; ++It)
 		{
-			It->ServerSetPeriodFrameForClass(Class, Distance);
+			It->ServerSetPeriodFrameForClass(Class, UE::LWC::FloatToIntCastChecked<int32>(Distance));
 		}
 	})
 );
diff --git a/Engine/Plugins/Runtime/ReplicationGraph/Source/Private/ReplicationGraphTypes.cpp b/Engine/Plugins/Runtime/ReplicationGraph/Source/Private/ReplicationGraphTypes.cpp
--- a/Engine/Plugins/Runtime/ReplicationGraph/Source/Private/ReplicationGraphTypes.cpp
+++ b/Engine/Plugins/Runtime/ReplicationGraph/Source/Private/ReplicationGraphTypes.cpp
@@ -554,7 +554,7 @@ void FActorRepListStatCollector::VisitStreamingLevelCollection(const UReplicatio
 	{
 		const uint32 ListSize = (uint32)LevelList.ReplicationActorList.RepList.Num();
 		const uint32 ListSlack = LevelList.ReplicationActorList.RepList.GetSlack();
-		const uint32 ListBytes = LevelList.ReplicationActorList.RepList.GetAllocatedSize();
+		const SIZE_T ListBytes = LevelList.ReplicationActorList.RepList.GetAllocatedSize();
 
 		{
 			FRepListStats& StreamingLevelStats = PerStreamingLevelStats.FindOrAdd(LevelList.StreamingLevelName);
@@ -580,7 +580,7 @@ void FActorRepListStatCollector::VisitExplicitStreamingLevelList(FName ListOwner
 {
 	const uint32 ListSize = (uint32)RepList.RepList.Num();
 	const uint32 ListSlack = RepList.RepList.GetSlack();
-	const uint32 ListBytes = RepList.RepList.GetAllocatedSize();
+	const SIZE_T ListBytes = RepList.RepList.GetAllocatedSize();
 
 	FRepListStats& ClassStats = PerClassStats.FindOrAdd(ListOwnerName);
 	ClassStats.NumLists++;
diff --git a/Engine/Plugins/Runtime/ReplicationGraph/Source/Public/ReplicationGraph.h b/Engine/Plugins/Runtime/ReplicationGraph/Source/Public/ReplicationGraph.h
--- a/Engine/Plugins/Runtime/ReplicationGraph/Source/Public/ReplicationGraph.h
+++ b/Engine/Plugins/Runtime/ReplicationGraph/Source/Public/ReplicationGraph.h
@@ -1031,7 +1031,7 @@ public:
 		check(NetUpdateFrequency != 0.0f);
 
 		// Replication Graph is frame based. Convert NetUpdateFrequency to ReplicationPeriodFrame based on Server MaxTickRate.
-		uint32 FramesBetweenUpdates = (uint32)FMath::RoundToInt(NetDriver->NetServerMaxTickRate / NetUpdateFrequency);
+		uint32 FramesBetweenUpdates = (uint32)FMath::RoundToInt(static_cast<float>(NetDriver->NetServerMaxTickRate) / NetUpdateFrequency);
 		FramesBetweenUpdates = FMath::Clamp<uint32>(FramesBetweenUpdates, 1, MAX_uint16);
 
 		return (uint16)FramesBetweenUpdates;
diff --git a/Engine/Plugins/Runtime/ReplicationGraph/Source/Public/ReplicationGraphTypes.h b/Engine/Plugins/Runtime/ReplicationGraph/Source/Public/ReplicationGraphTypes.h
--- a/Engine/Plugins/Runtime/ReplicationGraph/Source/Public/ReplicationGraphTypes.h
+++ b/Engine/Plugins/Runtime/ReplicationGraph/Source/Public/ReplicationGraphTypes.h
@@ -1296,7 +1296,7 @@ struct FPrioritizedActorFullDebugDetails
 	bool operator==(const FActorRepListType& InActor) const { return Actor == InActor; }
 
 	FActorRepListType Actor;
-	float DistanceSq = 0.f;
+	FVector::FReal DistanceSq = 0.f;
 	float DistanceFactor = 0.f;
 
 	uint32 FramesSinceLastRap = 0;
diff --git a/Engine/Plugins/Runtime/ReplicationGraph/Source/ReplicationGraph.Build.cs b/Engine/Plugins/Runtime/ReplicationGraph/Source/ReplicationGraph.Build.cs
--- a/Engine/Plugins/Runtime/ReplicationGraph/Source/ReplicationGraph.Build.cs
+++ b/Engine/Plugins/Runtime/ReplicationGraph/Source/ReplicationGraph.Build.cs
@@ -16,5 +16,7 @@ public class ReplicationGraph : ModuleRules
 				"PerfCounters"
 			}
 		);
+
+		UnsafeTypeCastWarningLevel = WarningLevel.Error;
 	}
 }
