Reverting uses of FloatCastChecked

diff --git a/Engine/Plugins/Runtime/ReplicationGraph/Source/Private/ReplicationGraph.cpp b/Engine/Plugins/Runtime/ReplicationGraph/Source/Private/ReplicationGraph.cpp
--- a/Engine/Plugins/Runtime/ReplicationGraph/Source/Private/ReplicationGraph.cpp
+++ b/Engine/Plugins/Runtime/ReplicationGraph/Source/Private/ReplicationGraph.cpp
@@ -1392,7 +1392,7 @@ void UReplicationGraph::ReplicateActorListsForConnections_Default(UNetReplicatio
 					// If no one is near this actor, skip it.
 					if (ViewersThatSkipActor >= Viewers.Num())
 					{
-						DO_REPGRAPH_DETAILS(PrioritizedReplicationList.GetNextSkippedDebugDetails(Actor)->DistanceCulled = FloatCastChecked<float>(FMath::Sqrt(SmallestDistanceSq), UE::LWC::DefaultFloatPrecision));
+						DO_REPGRAPH_DETAILS(PrioritizedReplicationList.GetNextSkippedDebugDetails(Actor)->DistanceCulled = static_cast<float>(FMath::Sqrt(SmallestDistanceSq)));
 
 						// Skipped actors should not have any 
 						if (bDoCulledOnConnectionCount)
@@ -1402,7 +1402,7 @@ void UReplicationGraph::ReplicateActorListsForConnections_Default(UNetReplicatio
 						continue;
 					}
 
-					const float DistanceFactor = FMath::Clamp(FloatCastChecked<float>(SmallestDistanceSq / MaxDistanceScaling, UE::LWC::DefaultFloatPrecision), 0.f, 1.f) * GlobalData.Settings.DistancePriorityScale;
+					const float DistanceFactor = FMath::Clamp(static_cast<float>(SmallestDistanceSq / MaxDistanceScaling), 0.f, 1.f) * GlobalData.Settings.DistancePriorityScale;
 					if (DO_REPGRAPH_DETAILS(UNLIKELY(DebugDetails)))
 					{
 						DebugDetails->DistanceSq = SmallestDistanceSq;
@@ -4111,9 +4111,9 @@ void UReplicationGraphNode_DynamicSpatialFrequency::CalcFrequencyForActor(AActor
 	// --------------------------------------------------------------------------------------------------------
 	const FVector& ConnectionViewDir = LowestDistanceViewer->ViewDir;
 	const FVector DirToActor = (GlobalInfo.WorldLocation - LowestDistanceViewer->ViewLocation);
-	const FVector::FReal DistanceToActor = FMath::Sqrt(SmallestDistanceToActorSq);
+	const float DistanceToActor = static_cast<float>(FMath::Sqrt(SmallestDistanceToActorSq));
 	const FVector NormDirToActor = DistanceToActor > SMALL_NUMBER ? (DirToActor / DistanceToActor) : DirToActor;
-	const float DotP = FMath::Clamp(FloatCastChecked<float>(FVector::DotProduct(NormDirToActor, ConnectionViewDir), UE::LWC::DefaultFloatPrecision), -1.0f, 1.0f);
+	const float DotP = FMath::Clamp(static_cast<float>(FVector::DotProduct(NormDirToActor, ConnectionViewDir)), -1.0f, 1.0f);
 
 	const bool ActorSupportsFastShared = (GlobalInfo.Settings.FastSharedReplicationFunc != nullptr);
 	TArrayView<FSpatializationZone>& ZoneList = ActorSupportsFastShared ? MySettings.ZoneSettings : MySettings.ZoneSettings_NonFastSharedActors;
@@ -4145,7 +4145,7 @@ void UReplicationGraphNode_DynamicSpatialFrequency::CalcFrequencyForActor(AActor
 				}
 
 				const float CullDist = ConnectionInfo.GetCullDistance();
-				const float DistPct = FloatCastChecked<float>(DistanceToActor / CullDist, UE::LWC::DefaultFloatPrecision);
+				const float DistPct = DistanceToActor / CullDist;
 
 				const float BiasDistPct = DistPct - ZoneInfo.MinDistPct;
 				const float FinalPCT = FMath::Clamp<float>( BiasDistPct / (ZoneInfo.MaxDistPct - ZoneInfo.MinDistPct), 0.f, 1.f);
