Handle gameplay debugger replicator in basic rep graph

diff --git a/Engine/Plugins/Runtime/ReplicationGraph/Source/Private/BasicReplicationGraph.cpp b/Engine/Plugins/Runtime/ReplicationGraph/Source/Private/BasicReplicationGraph.cpp
--- a/Engine/Plugins/Runtime/ReplicationGraph/Source/Private/BasicReplicationGraph.cpp
+++ b/Engine/Plugins/Runtime/ReplicationGraph/Source/Private/BasicReplicationGraph.cpp
@@ -4,6 +4,11 @@
 #include "UObject/UObjectIterator.h"
 #include "Engine/ChildConnection.h"
 
+#if WITH_GAMEPLAY_DEBUGGER
+#include "GameplayDebuggerCategoryReplicator.h"
+#include "GameFramework/PlayerController.h"
+#endif
+
 #include UE_INLINE_GENERATED_CPP_BY_NAME(BasicReplicationGraph)
 
 UBasicReplicationGraph::UBasicReplicationGraph()
@@ -48,6 +53,10 @@ void UBasicReplicationGraph::InitGlobalActorClassSettings()
 		
 		GlobalActorReplicationInfoMap.SetClassInfo( Class, ClassInfo );
 	}
+
+#if WITH_GAMEPLAY_DEBUGGER
+	AGameplayDebuggerCategoryReplicator::NotifyDebuggerOwnerChange.AddUObject(this, &ThisClass::OnGameplayDebuggerOwnerChange);
+#endif
 }
 
 void UBasicReplicationGraph::InitGlobalGraphNodes()
@@ -81,6 +90,14 @@ void UBasicReplicationGraph::InitConnectionGraphNodes(UNetReplicationGraphConnec
 
 void UBasicReplicationGraph::RouteAddNetworkActorToNodes(const FNewReplicatedActorInfo& ActorInfo, FGlobalActorReplicationInfo& GlobalInfo)
 {
+#if WITH_GAMEPLAY_DEBUGGER
+	// This is intended as an example and not shipping code, see the Lyra sample for a better way to handle routing with class routing mappings.
+	if (ActorInfo.Actor->IsA(AGameplayDebuggerCategoryReplicator::StaticClass()))
+	{
+		return;
+	}
+#endif
+
 	ensureMsgf((ActorInfo.Actor->bAlwaysRelevant && ActorInfo.Actor->bOnlyRelevantToOwner) == false, TEXT("Replicated actor %s is both bAlwaysRelevant and bOnlyRelevantToOwner. Only one can be supported."), *ActorInfo.Actor->GetName());
 	if (ActorInfo.Actor->bAlwaysRelevant)
 	{
@@ -100,6 +117,14 @@ void UBasicReplicationGraph::RouteAddNetworkActorToNodes(const FNewReplicatedAct
 
 void UBasicReplicationGraph::RouteRemoveNetworkActorToNodes(const FNewReplicatedActorInfo& ActorInfo)
 {
+#if WITH_GAMEPLAY_DEBUGGER
+	// This is intended as an example and not shipping code, see the Lyra sample for a better way to handle routing with class routing mappings.
+	if (ActorInfo.Actor->IsA(AGameplayDebuggerCategoryReplicator::StaticClass()))
+	{
+		return;
+	}
+#endif
+
 	if (ActorInfo.Actor->bAlwaysRelevant)
 	{
 		AlwaysRelevantNode->NotifyRemoveNetworkActor(ActorInfo);
@@ -180,6 +205,34 @@ int32 UBasicReplicationGraph::ServerReplicateActors(float DeltaSeconds)
 	return Super::ServerReplicateActors(DeltaSeconds);
 }
 
+#if WITH_GAMEPLAY_DEBUGGER
+void UBasicReplicationGraph::OnGameplayDebuggerOwnerChange(AGameplayDebuggerCategoryReplicator* Debugger, APlayerController* OldOwner)
+{
+	if (!Debugger || (Debugger->GetWorld() != GetWorld()))
+	{
+		return;
+	}
+
+	FNewReplicatedActorInfo ActorInfo(Debugger);
+
+	if (OldOwner)
+	{
+		if (UReplicationGraphNode* Node = OldOwner->GetNetConnection() ? GetAlwaysRelevantNodeForConnection(OldOwner->GetNetConnection()) : nullptr)
+		{
+			Node->NotifyRemoveNetworkActor(ActorInfo);
+		}
+	}
+
+	if (APlayerController* NewOwner = Debugger->GetReplicationOwner())
+	{
+		if (UReplicationGraphNode* Node = NewOwner->GetNetConnection() ? GetAlwaysRelevantNodeForConnection(NewOwner->GetNetConnection()) : nullptr)
+		{
+			Node->NotifyAddNetworkActor(ActorInfo);
+		}
+	}
+}
+#endif // WITH_GAMEPALY_DEBUGGER
+
 bool FConnectionAlwaysRelevantNodePair::operator==(UNetConnection* InConnection) const
 {
 	// Any children should be looking at their parent connections instead.
diff --git a/Engine/Plugins/Runtime/ReplicationGraph/Source/Public/BasicReplicationGraph.h b/Engine/Plugins/Runtime/ReplicationGraph/Source/Public/BasicReplicationGraph.h
--- a/Engine/Plugins/Runtime/ReplicationGraph/Source/Public/BasicReplicationGraph.h
+++ b/Engine/Plugins/Runtime/ReplicationGraph/Source/Public/BasicReplicationGraph.h
@@ -6,6 +6,9 @@
 #include "BasicReplicationGraph.generated.h"
 
 struct FNewReplicatedActorInfo;
+#if WITH_GAMEPLAY_DEBUGGER
+class AGameplayDebuggerCategoryReplicator;
+#endif
 
 USTRUCT()
 struct FConnectionAlwaysRelevantNodePair
@@ -62,8 +65,11 @@ public:
 	UPROPERTY()
 	TArray<TObjectPtr<AActor>> ActorsWithoutNetConnection;
 
-
 	UReplicationGraphNode_AlwaysRelevant_ForConnection* GetAlwaysRelevantNodeForConnection(UNetConnection* Connection);
+
+#if WITH_GAMEPLAY_DEBUGGER
+	void OnGameplayDebuggerOwnerChange(AGameplayDebuggerCategoryReplicator* Debugger, APlayerController* OldOwner);
+#endif
 };
 
 #if UE_ENABLE_INCLUDE_ORDER_DEPRECATED_IN_5_2
diff --git a/Engine/Plugins/Runtime/ReplicationGraph/Source/ReplicationGraph.Build.cs b/Engine/Plugins/Runtime/ReplicationGraph/Source/ReplicationGraph.Build.cs
--- a/Engine/Plugins/Runtime/ReplicationGraph/Source/ReplicationGraph.Build.cs
+++ b/Engine/Plugins/Runtime/ReplicationGraph/Source/ReplicationGraph.Build.cs
@@ -18,5 +18,7 @@ public class ReplicationGraph : ModuleRules
 		);
 
 		UnsafeTypeCastWarningLevel = WarningLevel.Error;
+
+		SetupGameplayDebuggerSupport(Target);
 	}
 }
diff --git a/Engine/Source/Runtime/GameplayDebugger/Public/GameplayDebuggerTypes.h b/Engine/Source/Runtime/GameplayDebugger/Public/GameplayDebuggerTypes.h
--- a/Engine/Source/Runtime/GameplayDebugger/Public/GameplayDebuggerTypes.h
+++ b/Engine/Source/Runtime/GameplayDebugger/Public/GameplayDebuggerTypes.h
@@ -333,7 +333,7 @@ struct GAMEPLAYDEBUGGER_API FGameplayDebuggerDataPack
 	/** get replication progress in (0..1) range */
 	FORCEINLINE float GetProgress() const
 	{
-		return (Header.DataOffset != Header.DataSize) ? ((1.0f * Header.DataOffset) / Header.DataSize) : 1.0f;
+		return (Header.DataOffset != Header.DataSize) ? (static_cast<float>(Header.DataOffset) / static_cast<float>(Header.DataSize)) : 1.0f;
 	}
 
 	/** is replication in progress? */
