Fix int to float cast error

diff --git a/Engine/Plugins/Runtime/ReplicationGraph/Source/Public/ReplicationGraph.h b/Engine/Plugins/Runtime/ReplicationGraph/Source/Public/ReplicationGraph.h
--- a/Engine/Plugins/Runtime/ReplicationGraph/Source/Public/ReplicationGraph.h
+++ b/Engine/Plugins/Runtime/ReplicationGraph/Source/Public/ReplicationGraph.h
@@ -1031,7 +1031,7 @@ public:
 		check(NetUpdateFrequency != 0.0f);
 
 		// Replication Graph is frame based. Convert NetUpdateFrequency to ReplicationPeriodFrame based on Server MaxTickRate.
-		uint32 FramesBetweenUpdates = (uint32)FMath::RoundToInt(NetDriver->GetNetServerMaxTickRate() / NetUpdateFrequency);
+		uint32 FramesBetweenUpdates = (uint32)FMath::RoundToInt(((float)NetDriver->GetNetServerMaxTickRate()) / NetUpdateFrequency);
 		FramesBetweenUpdates = FMath::Clamp<uint32>(FramesBetweenUpdates, 1, MAX_uint16);
 
 		return (uint16)FramesBetweenUpdates;
