RepGraph: Fix int to float cast warning

#rnx

diff --git a/Engine/Plugins/Runtime/ReplicationGraph/Source/Private/ReplicationGraph.cpp b/Engine/Plugins/Runtime/ReplicationGraph/Source/Private/ReplicationGraph.cpp
--- a/Engine/Plugins/Runtime/ReplicationGraph/Source/Private/ReplicationGraph.cpp
+++ b/Engine/Plugins/Runtime/ReplicationGraph/Source/Private/ReplicationGraph.cpp
@@ -3992,7 +3992,7 @@ void UReplicationGraphNode_DynamicSpatialFrequency::GatherActorListsForConnectio
 
 		if (CVar_RepGraph_DynamicSpatialFrequency_UncapBandwidth > 0)
 		{
-			UE_LOG(LogReplicationGraph, Display, TEXT("Uncapped bandwidth usage of UReplicationGraphNode_DynamicSpatialFrequency = %d bits -> %d bytes -> %.2f KBytes/sec"), BitsWritten, (BitsWritten+7) >> 3,  ((float)((BitsWritten+7)>>3)/1024.f) * GraphGlobals->ReplicationGraph->NetDriver->GetNetServerMaxTickRate());
+			UE_LOG(LogReplicationGraph, Display, TEXT("Uncapped bandwidth usage of UReplicationGraphNode_DynamicSpatialFrequency = %d bits -> %d bytes -> %.2f KBytes/sec"), BitsWritten, (BitsWritten+7) >> 3,  ((float)((BitsWritten+7)>>3)/1024.f) * ((float)GraphGlobals->ReplicationGraph->NetDriver->GetNetServerMaxTickRate()));
 		}
 	}
 #endif // WITH_SERVER_CODE
