Make CVar_RepGraph_DynamicSpatialFrequency_UncapBandwidth non-const in shipping builds so it can be hotfixed on/off

diff --git a/Engine/Plugins/Runtime/ReplicationGraph/Source/Private/ReplicationGraph.cpp b/Engine/Plugins/Runtime/ReplicationGraph/Source/Private/ReplicationGraph.cpp
--- a/Engine/Plugins/Runtime/ReplicationGraph/Source/Private/ReplicationGraph.cpp
+++ b/Engine/Plugins/Runtime/ReplicationGraph/Source/Private/ReplicationGraph.cpp
@@ -3802,7 +3802,10 @@ UReplicationGraphNode_DynamicSpatialFrequency::UReplicationGraphNode_DynamicSpat
 	CSVStatName = "DynamicSpatialFrequencyGatherPrioritize";
 }
 
-REPGRAPH_DEVCVAR_SHIPCONST(int32, "Net.RepGraph.DynamicSpatialFrequency.UncapBandwidth", CVar_RepGraph_DynamicSpatialFrequency_UncapBandwidth, 0, "Testing CVar that uncaps bandwidth on UReplicationGraphNode_DynamicSpatialFrequency nodes.");
+int32 CVar_RepGraph_DynamicSpatialFrequency_UncapBandwidth = 0;
+static FAutoConsoleVariableRef CVarRepGraphDynamicSpatialFrequencyUncapBandwidth(TEXT("Net.RepGraph.DynamicSpatialFrequency.UncapBandwidth"), CVar_RepGraph_DynamicSpatialFrequency_UncapBandwidth,
+	TEXT("Uncaps bandwidth on UReplicationGraphNode_DynamicSpatialFrequency nodes."), ECVF_Default);
+
 REPGRAPH_DEVCVAR_SHIPCONST(int32, "Net.RepGraph.DynamicSpatialFrequency.OpportunisticLoadBalance", CVar_RepGraph_DynamicSpatialFrequency_OpportunisticLoadBalance, 1, "Defers replication 1 frame in cases where many actors replicate on this frame but few on next frame.");
 
 FORCEINLINE bool ReplicatesEveryFrame(const FConnectionReplicationActorInfo& ConnectionInfo, const bool CheckFastPath)
