- Fixing assert on FloatToIntCastChecked when you set your actor cull distance larger than std::numerics<int>

diff --git a/Engine/Plugins/Runtime/ReplicationGraph/Source/Private/ReplicationGraph.cpp b/Engine/Plugins/Runtime/ReplicationGraph/Source/Private/ReplicationGraph.cpp
--- a/Engine/Plugins/Runtime/ReplicationGraph/Source/Private/ReplicationGraph.cpp
+++ b/Engine/Plugins/Runtime/ReplicationGraph/Source/Private/ReplicationGraph.cpp
@@ -5237,8 +5237,8 @@ UReplicationGraphNode_GridSpatialization2D::FActorCellInfo UReplicationGraphNode
 	const FVector::FReal LocationBiasY = (ClampedLocation.Y - SpatialBias.Y);
 
 	const FVector::FReal Dist = CullDistance;
-	const FVector::FReal MinX = LocationBiasX - Dist;
-	const FVector::FReal MinY = LocationBiasY - Dist;
+	const FVector::FReal MinX = FMath::Max(0, LocationBiasX - Dist);
+	const FVector::FReal MinY = FMath::Max(0, LocationBiasY - Dist);
 	FVector::FReal MaxX = LocationBiasX + Dist;
 	FVector::FReal MaxY = LocationBiasY + Dist;
 
