RepGraph
* Fix log spam caused by actors going outside the legal world bounds and staying OOB.
* Now we will only log a warning once when the actor goes OOB
* net.ReplicationGraph.PrintGraph now prints the list of DynamicSpatializedActors monitored by the Grid2D node.

diff --git a/Engine/Plugins/Runtime/ReplicationGraph/Source/Private/ReplicationGraph.cpp b/Engine/Plugins/Runtime/ReplicationGraph/Source/Private/ReplicationGraph.cpp
--- a/Engine/Plugins/Runtime/ReplicationGraph/Source/Private/ReplicationGraph.cpp
+++ b/Engine/Plugins/Runtime/ReplicationGraph/Source/Private/ReplicationGraph.cpp
@@ -5223,7 +5223,12 @@ UReplicationGraphNode_GridSpatialization2D::FActorCellInfo UReplicationGraphNode
 		Location3D.Y < -RepGraphHalfWorldMax || Location3D.Y > RepGraphHalfWorldMax ||
 		Location3D.Z < -RepGraphHalfWorldMax || Location3D.Z > RepGraphHalfWorldMax)
 	{
-		UE_LOG(LogReplicationGraph, Warning, TEXT("GetCellInfoForActor: Actor %s is outside world bounds with a location of %s. Clamping grid location to world bounds."), *GetFullNameSafe(Actor), *Location3D.ToString());
+		FGlobalActorReplicationInfo& ActorRepInfo = GraphGlobals->GlobalActorReplicationInfoMap->Get(Actor);
+		if (!ActorRepInfo.bWasWorldLocClamped)
+		{
+			UE_LOG(LogReplicationGraph, Warning, TEXT("GetCellInfoForActor: Actor %s is outside world bounds with a location of %s. Clamping grid location to world bounds."), *GetFullNameSafe(Actor), *Location3D.ToString());	
+			ActorRepInfo.bWasWorldLocClamped = true;
+		}
 		ClampedLocation = Location3D.BoundToCube(RepGraphHalfWorldMax);
 	}
 
diff --git a/Engine/Plugins/Runtime/ReplicationGraph/Source/Private/ReplicationGraphDebugging.cpp b/Engine/Plugins/Runtime/ReplicationGraph/Source/Private/ReplicationGraphDebugging.cpp
--- a/Engine/Plugins/Runtime/ReplicationGraph/Source/Private/ReplicationGraphDebugging.cpp
+++ b/Engine/Plugins/Runtime/ReplicationGraph/Source/Private/ReplicationGraphDebugging.cpp
@@ -1134,6 +1134,48 @@ void LogActorRepList(FReplicationGraphDebugInfo& DebugInfo, FString Prefix, cons
 	DebugInfo.Log(ActorListStr);
 }
 
+void UReplicationGraphNode_GridSpatialization2D::LogNode(FReplicationGraphDebugInfo& DebugInfo, const FString& NodeName) const
+{
+	DebugInfo.Log(NodeName);
+
+	DebugInfo.PushIndent();
+	
+	FString DynamicActorsHeader = FString::Printf(TEXT("Tracking %d dynamic actors:"), DynamicSpatializedActors.Num());
+	DebugInfo.Log(DynamicActorsHeader);	
+	
+	if (DebugInfo.Flags == FReplicationGraphDebugInfo::ShowActors)
+	{
+		FString ActorList; 
+		for (auto& MapIt : DynamicSpatializedActors)
+		{
+			ActorList += GetActorRepListTypeDebugString(MapIt.Key);
+			ActorList += TEXT(" ");
+		}
+
+		DebugInfo.Log(ActorList);
+	}
+	DebugInfo.PopIndent();
+
+	DebugInfo.PushIndent();
+	for (const UReplicationGraphNode* ChildNode : AllChildNodes)
+	{
+		if (DebugInfo.bShowEmptyNodes == false)
+		{
+			TArray<FActorRepListType> TempArray;
+			ChildNode->GetAllActorsInNode_Debugging(TempArray);
+			if (TempArray.Num() == 0)
+			{
+				continue;
+			}
+		}
+
+		ChildNode->LogNode(DebugInfo, ChildNode->GetDebugString());
+	}
+	DebugInfo.PopIndent();
+
+
+}
+
 void UReplicationGraphNode_GridCell::LogNode(FReplicationGraphDebugInfo& DebugInfo, const FString& NodeName) const
 {
 	DebugInfo.Log(NodeName);
diff --git a/Engine/Plugins/Runtime/ReplicationGraph/Source/Public/ReplicationGraph.h b/Engine/Plugins/Runtime/ReplicationGraph/Source/Public/ReplicationGraph.h
--- a/Engine/Plugins/Runtime/ReplicationGraph/Source/Public/ReplicationGraph.h
+++ b/Engine/Plugins/Runtime/ReplicationGraph/Source/Public/ReplicationGraph.h
@@ -555,6 +555,7 @@ public:
 	virtual void NotifyResetAllNetworkActors() override;
 	virtual void PrepareForReplication() override;
 	virtual void GatherActorListsForConnection(const FConnectionGatherActorListParameters& Params) override;
+	virtual void LogNode(FReplicationGraphDebugInfo& DebugInfo, const FString& NodeName) const override;
 
 	void AddActor_Static(const FNewReplicatedActorInfo& ActorInfo, FGlobalActorReplicationInfo& ActorRepInfo) { AddActorInternal_Static(ActorInfo, ActorRepInfo, false); }
 	void AddActor_Dynamic(const FNewReplicatedActorInfo& ActorInfo, FGlobalActorReplicationInfo& ActorRepInfo) { AddActorInternal_Dynamic(ActorInfo); }
diff --git a/Engine/Plugins/Runtime/ReplicationGraph/Source/Public/ReplicationGraphTypes.h b/Engine/Plugins/Runtime/ReplicationGraph/Source/Public/ReplicationGraphTypes.h
--- a/Engine/Plugins/Runtime/ReplicationGraph/Source/Public/ReplicationGraphTypes.h
+++ b/Engine/Plugins/Runtime/ReplicationGraph/Source/Public/ReplicationGraphTypes.h
@@ -700,7 +700,13 @@ struct FGlobalActorReplicationEvents
 struct FGlobalActorReplicationInfo
 {
 	FGlobalActorReplicationInfo(FClassReplicationInfo& ClassInfo) 
-		: LastPreReplicationFrame(0), WorldLocation(ForceInitToZero), Settings(ClassInfo) { }
+		: LastPreReplicationFrame(0)
+		, WorldLocation(ForceInitToZero)
+		, bWantsToBeDormant(false)
+		, bSwapRolesOnReplicate(false)
+		, bWasWorldLocClamped(false)
+		, Settings(ClassInfo) 
+	{ }
 
 	// -----------------------------------------------------------
 	//	Dynamic state
@@ -716,10 +722,13 @@ struct FGlobalActorReplicationInfo
 	FVector WorldLocation;
 
 	/** Mirrors AActor::NetDormancy > DORM_Awake */
-	bool bWantsToBeDormant = false;
+	uint32 bWantsToBeDormant : 1;
 
 	/** True if we should swap the actor role and remote role before calling ReplicateActor() */
-	bool bSwapRolesOnReplicate = false;
+	uint32 bSwapRolesOnReplicate : 1;
+
+	/** Set when the object is found with an invalid world location */
+	uint32 bWasWorldLocClamped : 1;
 
 	/** Class default mirrors: state that is initialized directly from class defaults (and can be later changed on a per-actor basis) */
 	FClassReplicationInfo Settings;
