Replication / Dormancy and Always Relevant: Fixed bug where an always relevant, dormant actor could lose relevancy on clients. Now that actor will not be destroyed on clients where they are always relevant.

#rnx

diff --git a/Engine/Plugins/Runtime/ReplicationGraph/Source/Private/ReplicationGraph.cpp b/Engine/Plugins/Runtime/ReplicationGraph/Source/Private/ReplicationGraph.cpp
--- a/Engine/Plugins/Runtime/ReplicationGraph/Source/Private/ReplicationGraph.cpp
+++ b/Engine/Plugins/Runtime/ReplicationGraph/Source/Private/ReplicationGraph.cpp
@@ -4385,12 +4385,8 @@ void UReplicationGraphNode_ConnectionDormancyNode::ConditionalGatherDormantActor
 		FConnectionReplicationActorInfo& ConnectionActorInfo = ConnectionActorInfoMap.FindOrAdd(Actor);
 		if (ConnectionActorInfo.bDormantOnConnection)
 		{
-			// If we trickled this actor, restore CullDistance to the default
-			if (ConnectionActorInfo.GetCullDistanceSquared() <= 0.f)
-			{
-				FGlobalActorReplicationInfo& GlobalInfo = GlobalActorReplicationInfoMap->Get(Actor);
-				ConnectionActorInfo.SetCullDistanceSquared(GlobalInfo.Settings.GetCullDistanceSquared());
-			}
+			// If we trickled this actor it doesn't need to be always relevant anymore.
+			ConnectionActorInfo.bForceCullDistanceToZero = false;
 
 			// It can be removed
 			ConnectionList.RemoveAtSwap(idx);
@@ -4400,11 +4396,11 @@ void UReplicationGraphNode_ConnectionDormancyNode::ConditionalGatherDormantActor
 			}
 
 			UE_CLOG(CVar_RepGraph_LogNetDormancyDetails > 0, LogReplicationGraph, Display, TEXT("GRAPH_DORMANCY: Actor %s is Dormant on %s. Removing from list. (%d elements left)"), *Actor->GetPathName(), *GetName(), ConnectionList.Num());
-			bShouldTrickle = false; // Dont trickle this frame because we are still encountering dormant actors
+			bShouldTrickle = false; // Don't trickle this frame because we are still encountering dormant actors
 		}
-		else if (CVar_RepGraph_TrickleDistCullOnDormancyNodes > 0 && bShouldTrickle)
+		else if (CVar_RepGraph_TrickleDistCullOnDormancyNodes > 0 && bShouldTrickle && ConnectionActorInfo.GetCullDistanceSquared() > 0.0f)
 		{
-			ConnectionActorInfo.SetCullDistanceSquared(0.f);
+			ConnectionActorInfo.bForceCullDistanceToZero = true;
 			bShouldTrickle = false; // trickle one actor per frame
 		}
 	}
@@ -5927,7 +5923,12 @@ void UReplicationGraphNode_GridSpatialization2D::GatherActorListsForConnection(c
 				{
 					if (ActorInfo->bDormantOnConnection)
 					{
-						Params.ConnectionManager.NotifyAddDormantDestructionInfo(Actor);
+						// If this actor is NOT set to always relevant, then flag this actor to be dormancy destroyed on this connection's client.
+						if (ActorInfo->GetCullDistanceSquared() > 0.0f)
+						{
+							Params.ConnectionManager.NotifyAddDormantDestructionInfo(Actor);
+						}
+
 						ActorInfo->bDormantOnConnection = false;
 						// Ideally, no actor info outside this list should be set to true, so we don't have to worry about resetting them.
 						// However we could consider iterating through the actor map to reset all of them.
diff --git a/Engine/Plugins/Runtime/ReplicationGraph/Source/Public/ReplicationGraphTypes.h b/Engine/Plugins/Runtime/ReplicationGraph/Source/Public/ReplicationGraphTypes.h
--- a/Engine/Plugins/Runtime/ReplicationGraph/Source/Public/ReplicationGraphTypes.h
+++ b/Engine/Plugins/Runtime/ReplicationGraph/Source/Public/ReplicationGraphTypes.h
@@ -1099,9 +1099,9 @@ ENUM_CLASS_FLAGS(FGlobalActorReplicationInfoMap::EWarnFlag);
 /** Per-Actor data that is stored per connection */
 struct FConnectionReplicationActorInfo
 {
-	FConnectionReplicationActorInfo() : bDormantOnConnection(0), bTearOff(0), bGridSpatilization_AlreadyDormant(0) { }
+	FConnectionReplicationActorInfo() : bDormantOnConnection(0), bTearOff(0), bGridSpatilization_AlreadyDormant(0), bForceCullDistanceToZero(0) { }
 
-	FConnectionReplicationActorInfo(const FGlobalActorReplicationInfo& GlobalInfo) : bDormantOnConnection(0), bTearOff(0), bGridSpatilization_AlreadyDormant(0)
+	FConnectionReplicationActorInfo(const FGlobalActorReplicationInfo& GlobalInfo) : bDormantOnConnection(0), bTearOff(0), bGridSpatilization_AlreadyDormant(0), bForceCullDistanceToZero(0)
 	{
 		// Pull data from the global actor info. This is done for things that we just want to duplicate in both places so that we can avoid a lookup into the global map
 		// and also for things that we want to be overridden per (connection/actor)
@@ -1131,8 +1131,8 @@ struct FConnectionReplicationActorInfo
 		CullDistance = FMath::Sqrt(CullDistanceSquared);
 	}
 
-	float GetCullDistance() const { return CullDistance; }
-	float GetCullDistanceSquared() const { return CullDistanceSquared; }
+	float GetCullDistance() const { return bForceCullDistanceToZero ? 0.0f : CullDistance; }
+	float GetCullDistanceSquared() const { return bForceCullDistanceToZero ? 0.0f : CullDistanceSquared; }
 
 	UActorChannel* Channel = nullptr;
 
@@ -1154,7 +1154,6 @@ public:
 	uint16	ReplicationPeriodFrame = 1;		
 	uint16	FastPath_ReplicationPeriodFrame = 1;
 	
-
 	/** The frame num that we will close the actor channel. This will get updated/pushed anytime the actor replicates based on FGlobalActorReplicationInfo::ActorChannelFrameTimeout  */
 	uint32 ActorChannelCloseFrameNum = 0;
 
@@ -1164,6 +1163,11 @@ public:
 	/** Used as an optimization when doing 2D Grid Spatilization, prevents replicating the dormancy of the same actor twice in splitscreen evaluations. */
 	uint8 bGridSpatilization_AlreadyDormant:1;
 
+	/** When enabled: GetCullDistance() and GetCullDistanceSquared() will return 0.0f.
+	*	When disabled: GetCullDistance() and GetCullDistanceSquared() return CullDistance and CullDistanceSquared respectively.
+	*	Useful for temporarily making the actor on this connection always relevant. */
+	uint8 bForceCullDistanceToZero:1;
+
 	void LogDebugString(FOutputDevice& Ar) const;
 };
 
