Change saturation analytics so that connection saturation, fast path saturation, and DSF saturation are all tracked separately

diff --git a/Engine/Plugins/Runtime/ReplicationGraph/Source/Private/ReplicationGraph.cpp b/Engine/Plugins/Runtime/ReplicationGraph/Source/Private/ReplicationGraph.cpp
--- a/Engine/Plugins/Runtime/ReplicationGraph/Source/Private/ReplicationGraph.cpp
+++ b/Engine/Plugins/Runtime/ReplicationGraph/Source/Private/ReplicationGraph.cpp
@@ -967,6 +967,7 @@ int32 UReplicationGraph::ServerReplicateActors(float DeltaSeconds)
 	FrameReplicationStats.Reset();
 
 	bWasConnectionSaturated = false;
+	bWasConnectionFastPathSaturated = false;
 
 	TSet<UNetConnection*> ConnectionsToClose;
 
@@ -1073,6 +1074,9 @@ int32 UReplicationGraph::ServerReplicateActors(float DeltaSeconds)
 		{
 			NetConnection->TrackReplicationForAnalytics(bWasConnectionSaturated);
 			bWasConnectionSaturated = false;
+
+			ConnectionManager->TrackReplicationForAnalytics(bWasConnectionFastPathSaturated);
+			bWasConnectionFastPathSaturated = false;
 		};
 
 		const FReplicationGraphDestructionSettings DestructionSettings(DestructInfoMaxDistanceSquared, CVar_RepGraph_OutOfRangeDistanceCheckRatio * DestructInfoMaxDistanceSquared);
@@ -1718,7 +1722,7 @@ void UReplicationGraph::ReplicateActorListsForConnections_FastShared(UNetReplica
 #endif
 			if (TotalBitsWritten > MaxBits)
 			{
-				NotifyConnectionSaturated(*ConnectionManager);
+				NotifyConnectionFastPathSaturated();
 				return;
 			}
 		}
@@ -2492,6 +2496,11 @@ void UReplicationGraph::NotifyConnectionSaturated(UNetReplicationGraphConnection
 	++GNumSaturatedConnections;
 }
 
+void UReplicationGraph::NotifyConnectionFastPathSaturated()
+{
+	bWasConnectionFastPathSaturated = true;
+}
+
 void UReplicationGraph::SetActorDestructionInfoToIgnoreDistanceCulling(AActor* DestroyedActor)
 {
 	if (!DestroyedActor)
@@ -2893,6 +2902,16 @@ void UNetReplicationGraphConnection::CleanupNodeCaches(const UReplicationGraphNo
 	PrevDormantActorListPerNode.Remove(Node);
 }
 
+void UNetReplicationGraphConnection::TrackReplicationForAnalytics(bool bFastPathSaturated)
+{
+	FastPathSaturationAnalytics.TrackReplication(bFastPathSaturated);
+}
+
+void UNetReplicationGraphConnection::ResetSaturationAnalytics()
+{
+	FastPathSaturationAnalytics = FRepGraphConnectionSaturationAnalytics();
+}
+
 void UNetReplicationGraphConnection::GetClientVisibleLevelNames(TSet<FName>& OutLevelNames) const
 {
 	if (NetConnection == nullptr)
@@ -3988,7 +4007,7 @@ void UReplicationGraphNode_DynamicSpatialFrequency::GatherActorListsForConnectio
 			// Bandwidth Cap
 			if (BitsWritten > MaxBits && CVar_RepGraph_DynamicSpatialFrequency_UncapBandwidth == 0)
 			{
-				RepGraph->NotifyConnectionSaturated(Params.ConnectionManager);
+				Params.ConnectionManager.NotifyDSFNodeSaturated(this);
 				break;
 			}
 		}
diff --git a/Engine/Plugins/Runtime/ReplicationGraph/Source/Public/ReplicationGraph.h b/Engine/Plugins/Runtime/ReplicationGraph/Source/Public/ReplicationGraph.h
--- a/Engine/Plugins/Runtime/ReplicationGraph/Source/Public/ReplicationGraph.h
+++ b/Engine/Plugins/Runtime/ReplicationGraph/Source/Public/ReplicationGraph.h
@@ -1024,6 +1024,8 @@ public:
 
 	void NotifyConnectionSaturated(class UNetReplicationGraphConnection& Connection);
 
+	void NotifyConnectionFastPathSaturated();
+
 	void SetActorDestructionInfoToIgnoreDistanceCulling(AActor* DestroyedActor);
 
 	uint16 GetReplicationPeriodFrameForFrequency(float NetUpdateFrequency) const
@@ -1102,6 +1104,9 @@ private:
 	/** Whether or not a connection was saturated during an update. */
 	bool bWasConnectionSaturated = false;
 
+	/** Whether or not all fast shared replication budget has been used during replication. */
+	bool bWasConnectionFastPathSaturated = false;
+
 protected:
 	/** Default Replication Path */
 	void ReplicateActorListsForConnections_Default(UNetReplicationGraphConnection* ConnectionManager, FGatheredReplicationActorLists& GatheredReplicationListsForConnection, FNetViewerArray& Viewers);
@@ -1191,6 +1196,28 @@ struct FLastLocationGatherInfo
 	}
 };
 
+struct FRepGraphConnectionSaturationAnalytics
+{
+	uint32 LongestRunOfSaturatedReplications = 0;
+	uint32 CurrentRunOfSaturatedReplications = 0;
+
+	void TrackReplication(bool bIsSaturated)
+	{
+		if (bIsSaturated)
+		{
+			++CurrentRunOfSaturatedReplications;
+			if (CurrentRunOfSaturatedReplications > LongestRunOfSaturatedReplications)
+			{
+				LongestRunOfSaturatedReplications = CurrentRunOfSaturatedReplications;
+			}
+		}
+		else
+		{
+			CurrentRunOfSaturatedReplications = 0;
+		}
+	}
+};
+
 /** Manages actor replication for a specific connection */
 UCLASS(transient)
 class REPLICATIONGRAPH_API UNetReplicationGraphConnection : public UReplicationConnectionDriver
@@ -1297,6 +1324,13 @@ public:
 
 	void CleanupNodeCaches(const UReplicationGraphNode* Node);
 
+	virtual void NotifyDSFNodeSaturated(const UReplicationGraphNode_DynamicSpatialFrequency* Node) {}
+
+	virtual void TrackReplicationForAnalytics(bool bFastPathSaturated);
+
+	virtual void ResetSaturationAnalytics();
+
+	FRepGraphConnectionSaturationAnalytics FastPathSaturationAnalytics;
 private:
 
 	// Stored list of dormant actors in a previous cell when it's been left - this is for
