## RepGraph actor reference validation

* Added ValidateActorReferences to RepGraph and RepGraphNode types
  * Iterates over all nodes in the RepGraph and checks for potentally stale actor references

## Performance

Tested on a Win64 Test Server
* ValidateActorReferences: 10-20us

diff --git a/Engine/Plugins/Runtime/ReplicationGraph/Source/Private/ReplicationGraph.cpp b/Engine/Plugins/Runtime/ReplicationGraph/Source/Private/ReplicationGraph.cpp
--- a/Engine/Plugins/Runtime/ReplicationGraph/Source/Private/ReplicationGraph.cpp
+++ b/Engine/Plugins/Runtime/ReplicationGraph/Source/Private/ReplicationGraph.cpp
@@ -3467,6 +3467,24 @@ void UReplicationGraphNode_ActorList::GetAllActorsInNode_Debugging(TArray<FActor
 	}
 }
 
+void UReplicationGraphNode_ActorList::VerifyActorReferencesInternal()
+{
+	Super::VerifyActorReferencesInternal();
+
+	for (const FActorRepListType& Actor : ReplicationActorList)
+	{
+		VerifyActorReference(Actor);
+	}
+
+	for (const FStreamingLevelActorListCollection::FStreamingLevelActors& Item : StreamingLevelCollection.StreamingLevelLists)
+	{
+		for (const FActorRepListType& Actor : Item.ReplicationActorList)
+		{
+			VerifyActorReference(Actor);
+		}
+	}
+}
+
 void UReplicationGraphNode_ActorList::TearDown()
 {
 	Super::TearDown();
@@ -3720,6 +3738,27 @@ void UReplicationGraphNode_ActorListFrequencyBuckets::GetAllActorsInNode_Debuggi
 	}
 }
 
+void UReplicationGraphNode_ActorListFrequencyBuckets::VerifyActorReferencesInternal()
+{
+	Super::VerifyActorReferencesInternal();
+
+	for (const FActorRepListRefView& List : NonStreamingCollection)
+	{
+		for (const FActorRepListType& Actor : List)
+		{
+			VerifyActorReference(Actor);
+		}
+	}
+
+	for (const auto& List : StreamingLevelCollection.StreamingLevelLists)
+	{
+		for (const auto& Actor : List.ReplicationActorList)
+		{
+			VerifyActorReference(Actor);
+		}
+	}
+}
+
 void UReplicationGraphNode_ActorListFrequencyBuckets::OnCollectActorRepListStats(FActorRepListStatCollector& StatsCollector) const
 {
 	for (const FActorRepListRefView& List : NonStreamingCollection)
@@ -4020,6 +4059,16 @@ void UReplicationGraphNode_DynamicSpatialFrequency::GatherActorListsForConnectio
 #endif // WITH_SERVER_CODE
 }
 
+void UReplicationGraphNode_DynamicSpatialFrequency::VerifyActorReferencesInternal()
+{
+	Super::VerifyActorReferencesInternal();
+
+	for (const FDynamicSpatialFrequency_SortedItem& Item : SortedReplicationList)
+	{
+		VerifyActorReference(Item.Actor);
+	}
+}
+
 REPGRAPH_DEVCVAR_SHIPCONST(int32, "Net.RepGraph.DynamicSpatialFrequency.Draw", CVar_RepGraph_DynamicSpatialFrequency_Draw, 0, "");
 
 int32 CVar_RepGraph_DynamicSpatialFrequency_ForceMaxFreq = 0;
@@ -4533,6 +4582,19 @@ void UReplicationGraphNode_ConnectionDormancyNode::NotifyResetAllNetworkActors()
 	RemovedStreamingLevelActorListCollection.Reset();
 }
 
+void UReplicationGraphNode_ConnectionDormancyNode::VerifyActorReferencesInternal()
+{
+	Super::VerifyActorReferencesInternal();
+
+	for (const FStreamingLevelActorListCollection::FStreamingLevelActors& List : RemovedStreamingLevelActorListCollection.StreamingLevelLists)
+	{
+		for (const FActorRepListType& Actor : List.ReplicationActorList)
+		{
+			VerifyActorReference(Actor);
+		}
+	}
+}
+
 bool UReplicationGraphNode_ConnectionDormancyNode::IsNodeObsolete(uint32 CurrentFrame) const
 {
 	// Test if the connection tied to this node has been destroyed
diff --git a/Engine/Plugins/Runtime/ReplicationGraph/Source/Private/ReplicationGraphDebugging.cpp b/Engine/Plugins/Runtime/ReplicationGraph/Source/Private/ReplicationGraphDebugging.cpp
--- a/Engine/Plugins/Runtime/ReplicationGraph/Source/Private/ReplicationGraphDebugging.cpp
+++ b/Engine/Plugins/Runtime/ReplicationGraph/Source/Private/ReplicationGraphDebugging.cpp
@@ -1048,6 +1048,19 @@ void UReplicationGraph::CollectRepListStats(FActorRepListStatCollector& StatColl
 	}
 }
 
+void UReplicationGraph::VerifyActorReferences()
+{
+	UE_LOG(LogReplicationGraph, Verbose, TEXT("UReplicationGraph::VerifyActorReferences - %s"),
+		*GetName());
+
+	VerifyActorReferencesInternal();
+
+	for (const TObjectPtr<UReplicationGraphNode>& Node : GlobalGraphNodes)
+	{
+		Node->VerifyActorReferences();
+	}
+}
+
 #if DO_ENABLE_REPGRAPH_DEBUG_ACTOR
 AReplicationGraphDebugActor* UReplicationGraph::CreateDebugActor() const
 {
@@ -1066,6 +1079,27 @@ void UReplicationGraphNode::GetAllActorsInNode_Debugging(TArray<FActorRepListTyp
 	}
 }
 
+void UReplicationGraphNode::VerifyActorReferences()
+{
+	VerifyActorReferencesInternal();
+
+	for (const TObjectPtr<UReplicationGraphNode>& Node : AllChildNodes)
+	{
+		Node->VerifyActorReferencesInternal();
+	}
+}
+
+bool UReplicationGraphNode::VerifyActorReference(const FActorRepListType& Actor) const
+{
+	const bool bIsValid = IsValid(Actor);
+	UE_CLOG(!bIsValid, LogReplicationGraph, Error,
+		TEXT("Invalid Actor %s (%s) is still referenced by %s"),
+		*GetNameSafe(Actor),
+		*GetNameSafe(Actor ? Actor->GetClass() : nullptr),
+		*GetName());
+	return bIsValid;
+}
+
 void UReplicationGraphNode::LogNode(FReplicationGraphDebugInfo& DebugInfo, const FString& NodeName) const
 {
 	DebugInfo.Log(NodeName);
@@ -1172,8 +1206,26 @@ void UReplicationGraphNode_GridSpatialization2D::LogNode(FReplicationGraphDebugI
 		ChildNode->LogNode(DebugInfo, ChildNode->GetDebugString());
 	}
 	DebugInfo.PopIndent();
+}
+
+void UReplicationGraphNode_GridSpatialization2D::VerifyActorReferencesInternal()
+{
+	Super::VerifyActorReferencesInternal();
+
+	for (const auto& Pair : DynamicSpatializedActors)
+	{
+		VerifyActorReference(Pair.Key);
+	}
 
+	for (const auto& Pair : StaticSpatializedActors)
+	{
+		VerifyActorReference(Pair.Key);
+	}
 
+	for (const FPendingStaticActors& Item : PendingStaticSpatializedActors)
+	{
+		VerifyActorReference(Item.Actor);
+	}
 }
 
 void UReplicationGraphNode_GridCell::LogNode(FReplicationGraphDebugInfo& DebugInfo, const FString& NodeName) const
@@ -1206,6 +1258,16 @@ void UReplicationGraphNode_TearOff_ForConnection::LogNode(FReplicationGraphDebug
 	DebugInfo.PopIndent();
 }
 
+void UReplicationGraphNode_TearOff_ForConnection::VerifyActorReferencesInternal()
+{
+	Super::VerifyActorReferencesInternal();
+
+	for (const FActorRepListType& Actor : ReplicationActorList)
+	{
+		VerifyActorReference(Actor);
+	}
+}
+
 void UReplicationGraphNode_TearOff_ForConnection::OnCollectActorRepListStats(FActorRepListStatCollector& StatsCollector) const
 {
 	StatsCollector.VisitRepList(this, ReplicationActorList);
diff --git a/Engine/Plugins/Runtime/ReplicationGraph/Source/Public/ReplicationGraph.h b/Engine/Plugins/Runtime/ReplicationGraph/Source/Public/ReplicationGraph.h
--- a/Engine/Plugins/Runtime/ReplicationGraph/Source/Public/ReplicationGraph.h
+++ b/Engine/Plugins/Runtime/ReplicationGraph/Source/Public/ReplicationGraph.h
@@ -94,6 +94,9 @@ public:
 	/** Debugging only function to return a normal TArray of actor rep list (for logging, debug UIs, etc) */
 	virtual void GetAllActorsInNode_Debugging(TArray<FActorRepListType>& OutArray) const;
 
+	/** Debugging only function to verify referenced actors (if any) are still valid. */
+	void VerifyActorReferences();
+
 	// -----------------------------------------------------
 
 	bool GetRequiresPrepareForReplication() const { return bRequiresPrepareForReplicationCall; }
@@ -150,6 +153,14 @@ protected:
 	 */
 	virtual void OnCollectActorRepListStats(struct FActorRepListStatCollector& StatsCollector) const {}
 
+
+	/**
+	 * Verify the validity of any potentially stale actor pointers. Used for debugging.
+	 */
+	virtual void VerifyActorReferencesInternal() {}
+
+	bool VerifyActorReference(const FActorRepListType& Actor) const;
+
 protected:
 
 	UPROPERTY()
@@ -184,6 +195,8 @@ public:
 
 	virtual void GetAllActorsInNode_Debugging(TArray<FActorRepListType>& OutArray) const override;
 
+	virtual void VerifyActorReferencesInternal() override;
+
 	virtual void TearDown() override;
 
 	/** Removes the actor very quickly but breaks the list order */
@@ -267,6 +280,8 @@ public:
 
 	virtual void GetAllActorsInNode_Debugging(TArray<FActorRepListType>& OutArray) const override;
 
+	virtual void VerifyActorReferencesInternal() override;
+
 	virtual void OnCollectActorRepListStats(struct FActorRepListStatCollector& StatsCollector) const override;
 
 	void SetNonStreamingCollectionSize(const int32 NewSize);
@@ -300,6 +315,8 @@ public:
 	UReplicationGraphNode_DynamicSpatialFrequency();
 	
 	virtual void GatherActorListsForConnection(const FConnectionGatherActorListParameters& Params) override;
+	
+	virtual void VerifyActorReferencesInternal() override;
 
 	// --------------------------------------------------------
 
@@ -363,7 +380,6 @@ protected:
 
 		bool operator<(const FDynamicSpatialFrequency_SortedItem& Other) const { return FramesTillReplicate < Other.FramesTillReplicate; }
 
-		UPROPERTY()
 		AActor* Actor = nullptr;
 
 		int32 FramesTillReplicate = 0; // Note this also serves as "Distance Sq" when doing the FSettings::MaxNearestActors pass.
@@ -423,6 +439,7 @@ public:
 	virtual void GatherActorListsForConnection(const FConnectionGatherActorListParameters& Params) override;
 	virtual bool NotifyRemoveNetworkActor(const FNewReplicatedActorInfo& ActorInfo, bool WarnIfNotFound) override;
 	virtual void NotifyResetAllNetworkActors() override;
+	virtual void VerifyActorReferencesInternal() override;
 
 	void NotifyActorDormancyFlush(FActorRepListType Actor);
 
@@ -556,6 +573,7 @@ public:
 	virtual void PrepareForReplication() override;
 	virtual void GatherActorListsForConnection(const FConnectionGatherActorListParameters& Params) override;
 	virtual void LogNode(FReplicationGraphDebugInfo& DebugInfo, const FString& NodeName) const override;
+	virtual void VerifyActorReferencesInternal() override;
 
 	void AddActor_Static(const FNewReplicatedActorInfo& ActorInfo, FGlobalActorReplicationInfo& ActorRepInfo) { AddActorInternal_Static(ActorInfo, ActorRepInfo, false); }
 	void AddActor_Dynamic(const FNewReplicatedActorInfo& ActorInfo, FGlobalActorReplicationInfo& ActorRepInfo) { AddActorInternal_Dynamic(ActorInfo); }
@@ -855,6 +873,7 @@ public:
 	virtual void NotifyResetAllNetworkActors() override { TearOffActors.Reset(); }
 	virtual void GatherActorListsForConnection(const FConnectionGatherActorListParameters& Params) override;
 	virtual void LogNode(FReplicationGraphDebugInfo& DebugInfo, const FString& NodeName) const override;
+	virtual void VerifyActorReferencesInternal() override;
 
 	void NotifyTearOffActor(AActor* Actor, uint32 FrameNum);
 
@@ -982,6 +1001,8 @@ public:
 
 	virtual void CollectRepListStats(struct FActorRepListStatCollector& StatCollector) const;
 
+	void VerifyActorReferences();
+
 	const TSharedPtr<FReplicationGraphGlobalData>& GetGraphGlobals() const { return GraphGlobals; }
 
 	uint32 GetReplicationGraphFrame() const { return ReplicationGraphFrame; }
@@ -1042,6 +1063,8 @@ public:
 
 	UNetReplicationGraphConnection* FindConnectionManager(UNetConnection* NetConnection) const;
 
+	virtual void VerifyActorReferencesInternal() {}
+
 protected:
 
 	virtual void InitializeForWorld(UWorld* World);
