RepGraph
* Added comment explaining why we always calculate actor distance to view when replicating

#rnx

diff --git a/Engine/Plugins/Runtime/ReplicationGraph/Source/Private/ReplicationGraph.cpp b/Engine/Plugins/Runtime/ReplicationGraph/Source/Private/ReplicationGraph.cpp
--- a/Engine/Plugins/Runtime/ReplicationGraph/Source/Private/ReplicationGraph.cpp
+++ b/Engine/Plugins/Runtime/ReplicationGraph/Source/Private/ReplicationGraph.cpp
@@ -1377,6 +1377,7 @@ void UReplicationGraph::ReplicateActorListsForConnections_Default(UNetReplicatio
 				// -------------------
 				if (GlobalData.Settings.DistancePriorityScale > 0.f)
 				{
+					// Always compute distance even for AlwaysRelevant actors since the priority scaling needs it.
 					FVector::FReal SmallestDistanceSq = std::numeric_limits<FVector::FReal>::max();
 					int32 ViewersThatSkipActor = 0;
 					
@@ -6007,7 +6008,7 @@ void UReplicationGraphNode_GridSpatialization2D::GatherActorListsForConnection(c
 						// If this actor is NOT set to always relevant, then flag this actor to be dormancy destroyed on this connection's client.
 						if (ActorInfo->GetCullDistanceSquared() > 0.0f)
 						{
-							Params.ConnectionManager.NotifyAddDormantDestructionInfo(Actor);
+						Params.ConnectionManager.NotifyAddDormantDestructionInfo(Actor);
 						}
 
 						ActorInfo->bDormantOnConnection = false;
