Null checking for Net Channel to prevent server crashes when we run out of Replication Channels

diff --git a/Engine/Plugins/Runtime/ReplicationGraph/Source/Private/ReplicationGraph.cpp b/Engine/Plugins/Runtime/ReplicationGraph/Source/Private/ReplicationGraph.cpp
--- a/Engine/Plugins/Runtime/ReplicationGraph/Source/Private/ReplicationGraph.cpp
+++ b/Engine/Plugins/Runtime/ReplicationGraph/Source/Private/ReplicationGraph.cpp
@@ -2409,6 +2409,21 @@ bool UReplicationGraph::ProcessRemoteFunction(class AActor* Actor, UFunction* Fu
 			}
 
 			Ch = (UActorChannel *)Connection->CreateChannelByName( NAME_Actor, EChannelCreateFlags::OpenedLocally );
+			if (!Ch)
+			{
+				if (!(Function->FunctionFlags & FUNC_NetReliable))
+				{
+					// Should we run out of channels then it should be ok to fail out for non-reliable functions
+					return true;
+				}
+				else
+				{
+					// if reliable we want to give user feedback before exiting to allow further hosting
+					Connection->Close(ENetCloseResult::ReplicationChannelCountMaxedOut);
+					return true;
+				}
+			}
+
 			Ch->SetChannelActor(Actor, ESetChannelActorFlags::None);
 			
 			if (UNetReplicationGraphConnection* ConnectionManager = Cast<UNetReplicationGraphConnection>(Connection->GetReplicationConnectionDriver()))
diff --git a/Engine/Source/Runtime/Net/Core/Public/Net/Core/Connection/NetCloseResult.h b/Engine/Source/Runtime/Net/Core/Public/Net/Core/Connection/NetCloseResult.h
--- a/Engine/Source/Runtime/Net/Core/Public/Net/Core/Connection/NetCloseResult.h
+++ b/Engine/Source/Runtime/Net/Core/Public/Net/Core/Connection/NetCloseResult.h
@@ -330,6 +330,8 @@ enum class ENetCloseResult : uint32
 	/** Field payload serialization failed */
 	FieldPayloadFail,
 
+	/** Replication Channel Count was exceeded and channel creation failed */
+	ReplicationChannelCountMaxedOut,
 
 	/** Beacon Error result types */
 
