Replication graph: reset bForceCullDistanceToZero when an actor is set to not dormant. Otherwise, if an actor is set to not dormant before it was handled in UReplicationGraphNode_ConnectionDormancyNode::ConditionalGatherDormantActorsForConnection, bForceCullDistanceToZero would remain true, effectively treating the actor as always relevant even if that wasn't the desired behavior.

diff --git a/Engine/Plugins/Runtime/ReplicationGraph/Source/Private/ReplicationGraph.cpp b/Engine/Plugins/Runtime/ReplicationGraph/Source/Private/ReplicationGraph.cpp
--- a/Engine/Plugins/Runtime/ReplicationGraph/Source/Private/ReplicationGraph.cpp
+++ b/Engine/Plugins/Runtime/ReplicationGraph/Source/Private/ReplicationGraph.cpp
@@ -3119,6 +3119,7 @@ void UNetReplicationGraphConnection::SetActorNotDormantOnConnection(AActor* InAc
 	{
 		Info->bDormantOnConnection = false;
 		Info->bGridSpatilization_AlreadyDormant = false;
+		Info->bForceCullDistanceToZero = false;
 		RemoveActorFromAllPrevDormantActorLists(InActor);
 	}
 }
