ReplicationGraph:
* Ensure when an actor is replicated with bNetTemporary set

diff --git a/Engine/Plugins/Runtime/ReplicationGraph/Source/Private/ReplicationGraph.cpp b/Engine/Plugins/Runtime/ReplicationGraph/Source/Private/ReplicationGraph.cpp
--- a/Engine/Plugins/Runtime/ReplicationGraph/Source/Private/ReplicationGraph.cpp
+++ b/Engine/Plugins/Runtime/ReplicationGraph/Source/Private/ReplicationGraph.cpp
@@ -681,6 +681,8 @@ void UReplicationGraph::AddNetworkActor(AActor* Actor)
 		return;
 	}
 
+	ensureMsgf(!Actor->bNetTemporary, TEXT("ReplicationGraph does not support bNetTemporary. Actor: %s has bNetTemporary set."), *Actor->GetPathName());
+
 	// Create global rep info	
 	FGlobalActorReplicationInfo& GlobalInfo = GlobalActorReplicationInfoMap.Get(Actor);
 	GlobalInfo.bWantsToBeDormant = Actor->NetDormancy > DORM_Awake;
