Iris - Force very high priority on controllers and view targets if cvar Net.ForceConnectionViewerPriority > 0. Deprecated cvar Net.RepGraph.ForceConnectionViewerPriority as the functionality is implemented by Iris as well. Moved said cvar to NetCore module.

diff --git a/Engine/Plugins/Runtime/ReplicationGraph/Source/Private/ReplicationGraph.cpp b/Engine/Plugins/Runtime/ReplicationGraph/Source/Private/ReplicationGraph.cpp
--- a/Engine/Plugins/Runtime/ReplicationGraph/Source/Private/ReplicationGraph.cpp
+++ b/Engine/Plugins/Runtime/ReplicationGraph/Source/Private/ReplicationGraph.cpp
@@ -46,6 +46,7 @@
 #include "Engine/ActorChannel.h"
 #include "GameFramework/PlayerController.h"
 #include "Net/RepLayout.h"
+#include "Net/Core/Misc/NetCVars.h"
 #include "Net/Core/Trace/Private/NetTraceInternal.h"
 #include "UObject/UObjectIterator.h"
 #include "DrawDebugHelpers.h"
@@ -132,8 +133,6 @@ static FAutoConsoleVariableRef CVarRepGraphOutOfRangeDistanceCheckRatio(TEXT("Ne
 int32 CVar_RepGraph_DormancyNode_ObsoleteBehavior = 0;
 static FAutoConsoleVariableRef CVarRepGraphDormancyNodeObsoleteBehavior(TEXT("Net.RepGraph.DormancyNodeObsoleteBehavior"), CVar_RepGraph_DormancyNode_ObsoleteBehavior, TEXT("This changes how the dormancy node deals with obsolete nodes. 0 = ignore. 1 = lazily destroy the node"), ECVF_Default);
 
-static TAutoConsoleVariable<float> CVar_ForceConnectionViewerPriority(TEXT("Net.RepGraph.ForceConnectionViewerPriority"), 1, TEXT("Force the connection's player controller and viewing pawn as topmost priority."));
-
 int32 CVar_RepGraph_GridSpatialization2D_DestroyDormantDynamicActorsDefault = 1;
 static FAutoConsoleVariableRef CVarRepGraphGridSpatialization2DDestroyDormantDynamicActorsDefault(TEXT("Net.RepGraph.GridSpatialization2DDestroyDormantDynamicActorsDefault"), CVar_RepGraph_GridSpatialization2D_DestroyDormantDynamicActorsDefault, TEXT("Configure what the default for UReplicationGraphNode_GridSpatialization2D::DestroyDormantDynamicActors should be."), ECVF_Default);
 
@@ -1479,7 +1478,7 @@ void UReplicationGraph::ReplicateActorListsForConnections_Default(UNetReplicatio
 					// We need to find if this is anyone's viewer or viewtarget, not just the parent connection.
 					if (Actor == CurViewer.ViewTarget || Actor == CurViewer.InViewer)
 					{
-						if (CVar_ForceConnectionViewerPriority.GetValueOnAnyThread() > 0)
+						if (UE::Net::CVar_ForceConnectionViewerPriority > 0)
 						{
 							AccumulatedPriority = -MAX_FLT;
 						}
diff --git a/Engine/Source/Runtime/Engine/Private/NetDriver.cpp b/Engine/Source/Runtime/Engine/Private/NetDriver.cpp
--- a/Engine/Source/Runtime/Engine/Private/NetDriver.cpp
+++ b/Engine/Source/Runtime/Engine/Private/NetDriver.cpp
@@ -11,6 +11,7 @@
 #include "EngineStats.h"
 #include "GameFramework/GameModeBase.h"
 #include "Misc/ConfigCacheIni.h"
+#include "Net/Core/NetHandle/NetHandleManager.h"
 #include "Net/Core/PropertyConditions/RepChangedPropertyTracker.h"
 #include "UObject/UObjectIterator.h"
 #include "Net/Core/Trace/Private/NetTraceInternal.h"
@@ -6550,12 +6551,14 @@ void UNetDriver::UpdateIrisReplicationViews() const
 			{
 				View.Pos = ViewTarget->GetActorLocation();
 				View.Dir = ViewTarget->GetActorRotation().Vector();
+				View.ViewTarget = UE::Net::FNetHandleManager::GetNetHandle(ViewTarget);
 			}
 			if (ViewingController)
 			{
 				FRotator ViewRotation = ViewingController->GetControlRotation();
 				ViewingController->GetPlayerViewPoint(View.Pos, ViewRotation);
 				View.Dir = ViewRotation.Vector();
+				View.Controller = UE::Net::FNetHandleManager::GetNetHandle(ViewingController);
 
 				if (const APlayerCameraManager* CameraManager = ViewingController->PlayerCameraManager)
 				{
diff --git a/Engine/Source/Runtime/Experimental/Iris/Core/Private/Iris/ReplicationSystem/Prioritization/ReplicationPrioritization.cpp b/Engine/Source/Runtime/Experimental/Iris/Core/Private/Iris/ReplicationSystem/Prioritization/ReplicationPrioritization.cpp
--- a/Engine/Source/Runtime/Experimental/Iris/Core/Private/Iris/ReplicationSystem/Prioritization/ReplicationPrioritization.cpp
+++ b/Engine/Source/Runtime/Experimental/Iris/Core/Private/Iris/ReplicationSystem/Prioritization/ReplicationPrioritization.cpp
@@ -4,6 +4,7 @@
 #include "Iris/IrisConstants.h"
 #include "Iris/Core/IrisLog.h"
 #include "Iris/Core/IrisProfiler.h"
+#include "Net/Core/Misc/NetCVars.h"
 #include "Net/Core/NetBitArray.h"
 #include "Iris/ReplicationSystem/NetRefHandleManager.h"
 #include "Iris/ReplicationSystem/ReplicationConnections.h"
@@ -734,6 +735,40 @@ void FReplicationPrioritization::PrioritizeForConnection(uint32 ConnId, FPriorit
 		checkf(false, TEXT("Unexpected BatchProcessStatus %u"), ProcessStatus);
 		break;
 	}
+
+	// Optionally force very high priority on view targets
+	if (CVar_ForceConnectionViewerPriority > 0)
+	{
+		SetHighPriorityOnViewTargets(MakeArrayView(ConnInfo.Priorities), PrioParameters.View);
+	}
+}
+
+void FReplicationPrioritization::SetHighPriorityOnViewTargets(const TArrayView<float>& Priorities, const FReplicationView& ReplicationView)
+{
+	using namespace UE::Net::Private;
+
+	// We allow a view target to appear multiple times. It will get the same priority regardless.
+	TArray<FNetHandle, TInlineAllocator<16>> ViewTargets;
+	for (const FReplicationView::FView& View : ReplicationView.Views)
+	{
+		if (View.Controller.IsValid())
+		{
+			ViewTargets.Add(View.Controller);
+		}
+		if (View.ViewTarget != View.Controller && View.ViewTarget.IsValid())
+		{
+			ViewTargets.Add(View.ViewTarget);
+		}
+	}
+
+	for (const FNetHandle NetHandle : ViewTargets)
+	{
+		const FInternalNetRefIndex ViewTargetInternalIndex = NetRefHandleManager->GetInternalIndexFromNetHandle(NetHandle);
+		if (ViewTargetInternalIndex != FNetRefHandleManager::InvalidInternalIndex)
+		{
+			Priorities[ViewTargetInternalIndex] = ViewTargetHighPriority;
+		}
+	}
 }
 
 /**
diff --git a/Engine/Source/Runtime/Experimental/Iris/Core/Private/Iris/ReplicationSystem/Prioritization/ReplicationPrioritization.h b/Engine/Source/Runtime/Experimental/Iris/Core/Private/Iris/ReplicationSystem/Prioritization/ReplicationPrioritization.h
--- a/Engine/Source/Runtime/Experimental/Iris/Core/Private/Iris/ReplicationSystem/Prioritization/ReplicationPrioritization.h
+++ b/Engine/Source/Runtime/Experimental/Iris/Core/Private/Iris/ReplicationSystem/Prioritization/ReplicationPrioritization.h
@@ -53,6 +53,7 @@ private:
 
 	void UpdatePrioritiesForNewAndDeletedObjects();
 	void PrioritizeForConnection(uint32 ConnId, FPrioritizerBatchHelper& BatchHelper, FNetBitArrayView Objects);
+	void SetHighPriorityOnViewTargets(const TArrayView<float>& Priorities, const FReplicationView& View);
 	void NotifyPrioritizersOfDirtyObjects(const FNetBitArrayView& DirtyObjectsThisFrame);
 	void BatchNotifyPrioritizersOfDirtyObjects(FUpdateDirtyObjectsBatchHelper& BatchHelper, uint32* ObjectIndices, uint32 ObjectCount);
 	void InitPrioritizers();
@@ -75,6 +76,7 @@ private:
 	};
 
 	static constexpr float DefaultPriority = 1.0f;
+	static constexpr float ViewTargetHighPriority = 1.0E7f;
 
 	TObjectPtr<const UReplicationSystem> ReplicationSystem;
 	FReplicationConnections* Connections = nullptr;
diff --git a/Engine/Source/Runtime/Experimental/Iris/Core/Public/Iris/ReplicationSystem/ReplicationView.h b/Engine/Source/Runtime/Experimental/Iris/Core/Public/Iris/ReplicationSystem/ReplicationView.h
--- a/Engine/Source/Runtime/Experimental/Iris/Core/Public/Iris/ReplicationSystem/ReplicationView.h
+++ b/Engine/Source/Runtime/Experimental/Iris/Core/Public/Iris/ReplicationSystem/ReplicationView.h
@@ -5,6 +5,7 @@
 #include "CoreTypes.h"
 #include "Containers/Array.h"
 #include "Math/Vector.h"
+#include "Net/Core/NetHandle/NetHandle.h"
 
 // Define this macro in your game's Target.cs and set it to the maximum clients per connection your game can have (ex: splitscreen).
 #ifndef UE_IRIS_INLINE_VIEWS_PER_CONNECTION
@@ -18,8 +19,15 @@ struct FReplicationView
 {
 	struct FView
 	{
+		/** The controlling net object associated with this view, typically a player controller. */
+		FNetHandle Controller;
+		/** The actor that is being directly viewed, usually a pawn. */
+		FNetHandle ViewTarget;
+		/** Where the viewer is looking from */
 		FVector Pos;
+		/** Direction the viewer is looking */
 		FVector Dir;
+		/** The field of view */
 		float FoVRadians;
 	};
 
diff --git a/Engine/Source/Runtime/Net/Core/Private/Net/Core/Misc/NetCVars.cpp b/Engine/Source/Runtime/Net/Core/Private/Net/Core/Misc/NetCVars.cpp
new file mode 100644
--- /dev/null
+++ b/Engine/Source/Runtime/Net/Core/Private/Net/Core/Misc/NetCVars.cpp
@@ -0,0 +1,15 @@
+// Copyright Epic Games, Inc. All Rights Reserved.
+
+#include "Net/Core/Misc/NetCvars.h"
+#include "HAL/IConsoleManager.h"
+
+namespace UE::Net
+{
+
+int32 CVar_ForceConnectionViewerPriority = 1;
+// Deprecated RepGraph specific name.
+static FAutoConsoleVariableRef CVarRepGraphForceConnectionViewerPriority(TEXT("Net.RepGraph.ForceConnectionViewerPriority"), CVar_ForceConnectionViewerPriority, TEXT("Force the connection's player controller and viewing pawn as topmost priority. Same as Net.ForceConnectionViewerPriority."));
+// New name without RepGraph as this is used by Iris as well.
+static FAutoConsoleVariableRef CVarForceConnectionViewerPriority(TEXT("Net.ForceConnectionViewerPriority"), CVar_ForceConnectionViewerPriority, TEXT("Force the connection's player controller and viewing pawn as topmost priority."));
+
+}
diff --git a/Engine/Source/Runtime/Net/Core/Public/Net/Core/Misc/NetCVars.h b/Engine/Source/Runtime/Net/Core/Public/Net/Core/Misc/NetCVars.h
new file mode 100644
--- /dev/null
+++ b/Engine/Source/Runtime/Net/Core/Public/Net/Core/Misc/NetCVars.h
@@ -0,0 +1,12 @@
+// Copyright Epic Games, Inc. All Rights Reserved.
+
+#pragma once
+
+#include "HAL/Platform.h"
+
+namespace UE::Net
+{
+
+NETCORE_API extern int32 CVar_ForceConnectionViewerPriority;
+
+}
