Fix C4702 warning

diff --git a/Engine/Plugins/Runtime/ReplicationGraph/Source/Public/ReplicationGraphTypes.h b/Engine/Plugins/Runtime/ReplicationGraph/Source/Public/ReplicationGraphTypes.h
--- a/Engine/Plugins/Runtime/ReplicationGraph/Source/Public/ReplicationGraphTypes.h
+++ b/Engine/Plugins/Runtime/ReplicationGraph/Source/Public/ReplicationGraphTypes.h
@@ -1666,8 +1666,9 @@ struct FReplicationGraphCSVTracker
 			VisibleLevelConnectionTracker.Add(VisibleLevelData);
 			return true;
 		}
-#endif
+#else
 		return false;
+#endif
 	}
 
 	void VisibleLevelConnectionRemoved(FName LevelName)
