Add cvar to flag net dormant actors that were destroyed by [REDACTED_PERSON] as not dormant anymore. This ensures the server initializes the flushed actor's FObjectReplicator correctly with the default state of the object, instead of the current state, since the client no longer has a current state.

#tests Bug repro with chests, and general gameplay

diff --git a/Engine/Plugins/Runtime/ReplicationGraph/Source/Private/ReplicationGraph.cpp b/Engine/Plugins/Runtime/ReplicationGraph/Source/Private/ReplicationGraph.cpp
--- a/Engine/Plugins/Runtime/ReplicationGraph/Source/Private/ReplicationGraph.cpp
+++ b/Engine/Plugins/Runtime/ReplicationGraph/Source/Private/ReplicationGraph.cpp
@@ -48,6 +48,7 @@
 #include "Net/RepLayout.h"
 #include "Net/Core/Misc/NetCVars.h"
 #include "Net/Core/Trace/Private/NetTraceInternal.h"
+#include "Engine/NetworkObjectList.h"
 #include "UObject/UObjectIterator.h"
 #include "DrawDebugHelpers.h"
 #include "Misc/ScopeExit.h"
@@ -126,6 +127,10 @@ int32 CVar_RepGraph_ReplicatedDormantDestructionInfosPerFrame = MAX_int32;
 static FAutoConsoleVariableRef CVarRepGraphReplicatedDormantDestructionInfosPerFrame(TEXT("Net.RepGraph.ReplicatedDormantDestructionInfosPerFrame"), CVar_RepGraph_ReplicatedDormantDestructionInfosPerFrame,
 	TEXT("If CVarRepGraphDormantDynamicActorsDestruction is true, this is the max number of destruction infos sent to a client per frame"), ECVF_Default);
 
+int32 CVar_RepGraph_MarkDormantDestroyedActorsActive = 1;
+static FAutoConsoleVariableRef CVarRepGraphMarkDormantDestroyedActorsActive(TEXT("Net.RepGraph.MarkDormantDestroyedActorsActive"), CVar_RepGraph_MarkDormantDestroyedActorsActive,
+	TEXT("If CVarRepGraphDormantDynamicActorsDestruction is true, mark destroyed dormant actors active to keep their dormancy state with the NetDriver in sync"), ECVF_Default);
+
 float CVar_RepGraph_OutOfRangeDistanceCheckRatio = 0.5f;
 static FAutoConsoleVariableRef CVarRepGraphOutOfRangeDistanceCheckRatio(TEXT("Net.RepGraph.OutOfRangeDistanceCheckRatio"), CVar_RepGraph_OutOfRangeDistanceCheckRatio,
 	TEXT("The ratio of DestructInfoMaxDistance that gives the distance traveled before we reevaluate the out of range destroyed actors list"), ECVF_Default);
@@ -2842,6 +2847,16 @@ void UNetReplicationGraphConnection::NotifyAddDormantDestructionInfo(AActor* Act
 			Info.Level = Level;
 			Info.ObjOuter = Actor->GetOuter();
 			Info.PathName = Actor->GetName();
+
+			if (CVar_RepGraph_MarkDormantDestroyedActorsActive)
+			{
+				// MarkActive here to get this connection off the actor's DormantConnections list.
+				// Since the actor will be destroyed, it shouldn't be considered dormant anymore.
+				// If FlushNetDormancy is called later for this actor, this prevents the flush
+				// from creating a replicator based on the actor's current state - which wouldn't
+				// work, becuase the client won't have the current state.
+				NetConnection->Driver->GetNetworkObjectList().MarkActive(Actor, NetConnection, NetConnection->Driver);
+			}
 		}
 	}
 #endif // WITH_SERVER_CODE
