Add facility to safety-check raw AActor pointers used by ReplicationGraph.

- Attempting to avoid a rare crash and log more info about the circumstances.

diff --git a/Engine/Plugins/Runtime/ReplicationGraph/Source/Private/ReplicationGraph.cpp b/Engine/Plugins/Runtime/ReplicationGraph/Source/Private/ReplicationGraph.cpp
--- a/Engine/Plugins/Runtime/ReplicationGraph/Source/Private/ReplicationGraph.cpp
+++ b/Engine/Plugins/Runtime/ReplicationGraph/Source/Private/ReplicationGraph.cpp
@@ -144,6 +144,11 @@ static FAutoConsoleVariableRef CVarRepGraphGridSpatialization2DDestroyDormantDyn
 int32 CVar_RepGraph_ConnectionHeavyComputationAmortization = 0;
 static FAutoConsoleVariableRef CVarRepGraphConnectionHeavyComputationAmortization(TEXT("Net.RepGraph.ConnectionHeavyComputationAmortization"), CVar_RepGraph_ConnectionHeavyComputationAmortization, TEXT("Non zero values enable heavy computation cost amortization by updating one connection per frame."), ECVF_Default);
 
+int32 CVar_RepGraph_LogDebugInfoPeriod = 60 * 5;
+static FAutoConsoleVariableRef CVarRepGraphLogDebugInfoPeriod(TEXT("Net.RepGraph.LogDebugInfoPeriod"), CVar_RepGraph_LogDebugInfoPeriod,
+	TEXT("How long to wait between logging debug info on invalid actors."), ECVF_Default);
+
+
 REPGRAPH_DEVCVAR_SHIPCONST(int32, "Net.RepGraph.LogNetDormancyDetails", CVar_RepGraph_LogNetDormancyDetails, 0, "Logs actors that are removed from the replication graph/nodes.");
 REPGRAPH_DEVCVAR_SHIPCONST(int32, "Net.RepGraph.LogActorRemove", CVar_RepGraph_LogActorRemove, 0, "Logs actors that are removed from the replication graph/nodes.");
 REPGRAPH_DEVCVAR_SHIPCONST(int32, "Net.RepGraph.LogActorAdd", CVar_RepGraph_LogActorAdd, 0, "Logs actors that are added to replication graph/nodes.");
@@ -237,6 +242,61 @@ void UpdateActorConnectionCounter(AActor* InActor, UNetConnection* InConnection,
 }
 #endif // WITH_SERVER_CODE
 
+/** Called by IsActorValidForReplication_LogMoreInfo to print more info on a failure */
+void LogMoreInfoOnIsActorValidFailure(const FActorRepListType& In)
+{
+#if WITH_SERVER_CODE
+	// Always log information about the actor (it's not much)
+	if (!DoesActorPointerLookValid(In))
+	{
+#if UE_ACTOR_REPLIST_TYPE_EXTRA_SAFETY
+		UE_LOG(LogReplicationGraph, Error, TEXT("Invalid actor pointer detected during replication: Ptr = %p, Name='%s', Owner='%s', OuterPackage='%s'"),
+			In.ActorRaw, *In.ActorName.ToString(), *In.OwnerName.ToString(), *In.OuterPackageName.ToString()
+			);
+#else
+		UE_LOG(LogReplicationGraph, Error, TEXT("Invalid actor detected during replication: Ptr = %p"), static_cast<AActor*>(In));
+#endif
+	}
+	else
+	{
+		// Actor pointer is valid, but some of its properties may be not 
+		AActor* Actor = In;
+		UE_LOG(LogReplicationGraph, Error, TEXT("Actor not valid for replication (BeingDestroyed:%d) (IsValid:%d) (Unreachable:%d) (TearOff:%d)! Actor = %s"),
+			Actor->IsActorBeingDestroyed(), IsValid(Actor), Actor->IsUnreachable(), Actor->GetTearOff(),
+			*Actor->GetFullName());
+#if UE_ACTOR_REPLIST_TYPE_EXTRA_SAFETY
+		UE_LOG(LogReplicationGraph, Error, TEXT("More info about invalid actor '%s': Owner='%s', OuterPackage='%s'"),
+			*In.ActorName.ToString(), *In.OwnerName.ToString(), *In.OuterPackageName.ToString());
+#endif
+	}
+
+
+	static std::atomic<double> LastTimeLogged = 0;
+	double CurrentTime = FPlatformTime::Seconds();
+	if (CVar_RepGraph_LogDebugInfoPeriod > 0 && ((CurrentTime - LastTimeLogged) > double(CVar_RepGraph_LogDebugInfoPeriod)))
+	{
+		LastTimeLogged = CurrentTime;
+
+		if (DoesActorPointerLookValid(In))
+		{
+			AActor* Actor = In;
+			UE_LOG(LogReplicationGraph, Log, TEXT("Invalid actor found at the last moment, executing Net.RepGraph.PrintAll:"));
+			GEngine->Exec(Actor->GetWorld(), TEXT("Net.RepGraph.PrintAll"));
+		}
+		else
+		{
+			UE_LOG(LogReplicationGraph, Warning, TEXT("Invalid actor found at the last moment, not executing Net.RepGraph.PrintAll due to invalid actor pointer."));
+		}
+
+		UE_LOG(LogReplicationGraph, Log, TEXT("Invalid actor found at the last moment, printing (game-specific) current routing:"));
+		for (TObjectIterator<UReplicationGraph> It; It; ++It)
+		{
+			It->PrintGraphDebugInfo_OnInvalidActor();
+		}
+	}
+#endif // WITH_SERVER_CODE
+}
+
 // --------------------------------------------------------------------------------------------------------------------------------------------
 // --------------------------------------------------------------------------------------------------------------------------------------------
 // --------------------------------------------------------------------------------------------------------------------------------------------
@@ -1326,10 +1386,15 @@ void UReplicationGraph::ReplicateActorListsForConnections_Default(UNetReplicatio
 		{
 			// Add actors from gathered list
 			NumGatheredActorsOnConnection += List.Num();
-			for (AActor* Actor : List)
+			for (const FActorRepListType& Actor : List)
 			{
 				RG_QUICK_SCOPE_CYCLE_COUNTER(Prioritize_InnerLoop);
 
+				if (!ensureMsgf(IsActorValidForReplication_LogMoreInfo(Actor), TEXT("Actor not valid for replication")))
+				{
+					continue;
+				}
+
 				// -----------------------------------------------------------------------------------------------------------------
 				//	Prioritize Actor for Connection: this is the main block of code for calculating a final score for this actor
 				//		-This is still pretty rough. It would be nice if this was customizable per project without suffering virtual calls.
@@ -1866,7 +1931,7 @@ int64 UReplicationGraph::ReplicateSingleActor_FastShared(AActor* Actor, FConnect
 #endif // WITH_SERVER_CODE
 }
 
-int64 UReplicationGraph::ReplicateSingleActor(AActor* Actor, FConnectionReplicationActorInfo& ActorInfo, FGlobalActorReplicationInfo& GlobalActorInfo, FPerConnectionActorInfoMap& ConnectionActorInfoMap, UNetReplicationGraphConnection& ConnectionManager, const uint32 FrameNum)
+int64 UReplicationGraph::ReplicateSingleActor(const FActorRepListType& Actor, FConnectionReplicationActorInfo& ActorInfo, FGlobalActorReplicationInfo& GlobalActorInfo, FPerConnectionActorInfoMap& ConnectionActorInfoMap, UNetReplicationGraphConnection& ConnectionManager, const uint32 FrameNum)
 {
 #if WITH_SERVER_CODE
 	RG_QUICK_SCOPE_CYCLE_COUNTER(NET_ReplicateActors_ReplicateSingleActor);
@@ -1886,9 +1951,7 @@ int64 UReplicationGraph::ReplicateSingleActor(AActor* Actor, FConnectionReplicat
 		UE_LOG(LogReplicationGraph, Display, TEXT("UReplicationGraph::ReplicateSingleActor: %s. NetConnection: %s"), *Actor->GetName(), *NetConnection->Describe());
 	}
 
-	if (!ensureMsgf(IsActorValidForReplication(Actor), TEXT("Actor not valid for replication (BeingDestroyed:%d) (IsValid:%d) (Unreachable:%d) (TearOff:%d)! Actor = %s, Channel = %s"),
-					Actor->IsActorBeingDestroyed(), IsValid(Actor), Actor->IsUnreachable(), Actor->GetTearOff(),
-					*Actor->GetFullName(), *DescribeSafe(ActorInfo.Channel)))
+	if (!ensureMsgf(IsActorValidForReplication_LogMoreInfo(Actor), TEXT("Actor not valid for replication")))
 	{
 		return 0;
 	}
@@ -4483,6 +4546,10 @@ void UReplicationGraphNode_ConnectionDormancyNode::ConditionalGatherDormantActor
 	for (int32 idx = ConnectionList.Num()-1; idx >= 0; --idx)
 	{
 		FActorRepListType Actor = ConnectionList[idx];
+		if (!ensureMsgf(IsActorValidForReplication_LogMoreInfo(Actor), TEXT("Actor not valid for replication")))
+		{
+			continue;
+		}
 		FConnectionReplicationActorInfo& ConnectionActorInfo = ConnectionActorInfoMap.FindOrAdd(Actor);
 		if (ConnectionActorInfo.bDormantOnConnection)
 		{
@@ -5244,7 +5311,7 @@ void UReplicationGraphNode_GridSpatialization2D::RemoveActorInternal_Static(cons
 	if (CVar_RepGraph_Verify)
 	{
 		// Verify this actor is in no nodes. This is pretty slow!
-		TArray<AActor*> AllActors;
+		TArray<FActorRepListType> AllActors;
 		for (auto& InnerArray : Grid)
 		{
 			for (UReplicationGraphNode_GridCell* N : InnerArray)
diff --git a/Engine/Plugins/Runtime/ReplicationGraph/Source/Private/ReplicationGraphDebugging.cpp b/Engine/Plugins/Runtime/ReplicationGraph/Source/Private/ReplicationGraphDebugging.cpp
--- a/Engine/Plugins/Runtime/ReplicationGraph/Source/Private/ReplicationGraphDebugging.cpp
+++ b/Engine/Plugins/Runtime/ReplicationGraph/Source/Private/ReplicationGraphDebugging.cpp
@@ -107,7 +107,7 @@ void AReplicationGraphDebugActor::ServerStartDebugging_Implementation()
 	int32 TotalNumCells = 0; // How many cells have been allocated
 	int32 TotalLeafNodes = 0; // How many cells have leaf nodes allocated
 
-	TSet<AActor*> UniqueActors;
+	TSet<FActorRepListType> UniqueActors;
 	int32 TotalElementsInLists = 0;
 
 	TMap<int32, int32> NumStreamLevelsMap;
@@ -425,7 +425,7 @@ void AReplicationGraphDebugActor::ServerCellInfo_Implementation()
 
 		LocationsSent.Add(CellLocation);
 
-		TArray<AActor*> ActorsInCell;
+		TArray<FActorRepListType> ActorsInCell;
 		FVector CellExtent(GridNode->CellSize, GridNode->CellSize, 10.f);
 
 		if (GridNode->Grid.IsValidIndex(CellX))
@@ -440,7 +440,16 @@ void AReplicationGraphDebugActor::ServerCellInfo_Implementation()
 			}
 		}
 
+#if !UE_ACTOR_REPLIST_TYPE_EXTRA_SAFETY
 		ClientCellInfo(CellLocation, CellExtent, ActorsInCell);
+#else
+		TArray<AActor*> RealActorsInCell;
+		for (const FActorRepListType& RepListActor : ActorsInCell)
+		{
+			RealActorsInCell.Add(RepListActor);
+		}
+		ClientCellInfo(CellLocation, CellExtent, RealActorsInCell);
+#endif
 	}	
 }
 
@@ -1344,11 +1353,14 @@ void PrintPrioritizedList(FOutputDevice& Ar, UNetReplicationGraphConnection* Con
 		FString StarvedString = bWasStarved ? FString::Printf(TEXT(" (Starved %d) "), RepFrameNum - ActorInfo.LastRepFrameNum) : TEXT("");
 
 #if REPGRAPH_DETAILS
-			
-		if (FPrioritizedActorFullDebugDetails* FullDetails = PrioritizedList->FullDebugDetails.Get()->FindByKey(Item.Actor))
+
+		if (PrioritizedList->FullDebugDetails.Get())
 		{
-			Ar.Logf(TEXT("%-40s %.4f %s %s"), *GetActorRepListTypeDebugString(Item.Actor), Item.Priority, *FullDetails->BuildString(), *StarvedString);
-			continue;
+			if (FPrioritizedActorFullDebugDetails* FullDetails = PrioritizedList->FullDebugDetails.Get()->FindByKey(Item.Actor))
+			{
+				Ar.Logf(TEXT("%-40s %.4f %s %s"), *GetActorRepListTypeDebugString(Item.Actor), Item.Priority, *FullDetails->BuildString(), *StarvedString);
+				continue;
+			}
 		}
 #endif
 
diff --git a/Engine/Plugins/Runtime/ReplicationGraph/Source/Private/ReplicationGraphTypes.cpp b/Engine/Plugins/Runtime/ReplicationGraph/Source/Private/ReplicationGraphTypes.cpp
--- a/Engine/Plugins/Runtime/ReplicationGraph/Source/Private/ReplicationGraphTypes.cpp
+++ b/Engine/Plugins/Runtime/ReplicationGraph/Source/Private/ReplicationGraphTypes.cpp
@@ -677,8 +677,19 @@ void FLevelBasedActorList::AppendAllLists(FGatheredReplicationActorLists& OutGat
 
 void FLevelBasedActorList::GetAllActors(TArray<AActor*>& OutAllActors) const
 {
+#if !UE_ACTOR_REPLIST_TYPE_EXTRA_SAFETY
 	PermanentLevelActors.AppendToTArray(OutAllActors);
 	StreamingLevelActors.GetAll_Debug(OutAllActors);
+#else
+	TArray<FActorRepListType> AllActors;
+	PermanentLevelActors.AppendToTArray(AllActors);
+	StreamingLevelActors.GetAll_Debug(AllActors);
+
+	for (const FActorRepListType& RepListActor : AllActors)
+	{
+		OutAllActors.Add(RepListActor);
+	}
+#endif
 }
 
 void FLevelBasedActorList::CountBytes(FArchive& Ar) const
diff --git a/Engine/Plugins/Runtime/ReplicationGraph/Source/Public/ReplicationGraph.h b/Engine/Plugins/Runtime/ReplicationGraph/Source/Public/ReplicationGraph.h
--- a/Engine/Plugins/Runtime/ReplicationGraph/Source/Public/ReplicationGraph.h
+++ b/Engine/Plugins/Runtime/ReplicationGraph/Source/Public/ReplicationGraph.h
@@ -959,6 +959,9 @@ public:
 	/** Sets the global and connection-specific cull distance setting of this actor */
 	void SetAllCullDistanceSettingsForActor(const FActorRepListType& Actor, float CullDistanceSquared);
 
+	/** Debugging facility called when invalid actor is being detected on one of many replication paths. */
+	virtual void PrintGraphDebugInfo_OnInvalidActor() {}
+
 	// --------------------------------------------------------------
 
 	/** Creates a new node for the graph. This and UReplicationNode::CreateChildNode should be the only things that create the graph node UObjects */
@@ -1038,7 +1041,7 @@ public:
 
 	virtual void ReplicateActorsForConnection(UNetConnection* NetConnection, FPerConnectionActorInfoMap& ConnectionActorInfoMap, UNetReplicationGraphConnection* ConnectionManager, const uint32 FrameNum);
 
-	int64 ReplicateSingleActor(AActor* Actor, FConnectionReplicationActorInfo& ActorInfo, FGlobalActorReplicationInfo& GlobalActorInfo, FPerConnectionActorInfoMap& ConnectionActorInfoMap, UNetReplicationGraphConnection& ConnectionManager, const uint32 FrameNum);
+	int64 ReplicateSingleActor(const FActorRepListType& Actor, FConnectionReplicationActorInfo& ActorInfo, FGlobalActorReplicationInfo& GlobalActorInfo, FPerConnectionActorInfoMap& ConnectionActorInfoMap, UNetReplicationGraphConnection& ConnectionManager, const uint32 FrameNum);
 	int64 ReplicateSingleActor_FastShared(AActor* Actor, FConnectionReplicationActorInfo& ConnectionData, FGlobalActorReplicationInfo& GlobalActorInfo, UNetReplicationGraphConnection& ConnectionManager, const uint32 FrameNum);
 
 	void UpdateActorChannelCloseFrameNum(AActor* Actor, FConnectionReplicationActorInfo& ConnectionData, const FGlobalActorReplicationInfo& GlobalData, const uint32 FrameNum, UNetConnection* NetConnection) const;
diff --git a/Engine/Plugins/Runtime/ReplicationGraph/Source/Public/ReplicationGraphTypes.h b/Engine/Plugins/Runtime/ReplicationGraph/Source/Public/ReplicationGraphTypes.h
--- a/Engine/Plugins/Runtime/ReplicationGraph/Source/Public/ReplicationGraphTypes.h
+++ b/Engine/Plugins/Runtime/ReplicationGraph/Source/Public/ReplicationGraphTypes.h
@@ -67,7 +67,119 @@ LLM_DECLARE_TAG_API(NetRepGraph, REPLICATIONGRAPH_API);
 // fast into arrays etc. (Currently we are using TMaps for associative data, static arrays would be faster but introduce constraints and headaches)
 // So for now, using a typedef and some helper functions to call out the interface/usage of FActorRepListType.
 
+// This define embeds an extra FObjectKey for a validity check
+#ifndef UE_ENABLE_ACTOR_REPLIST_TRACKING
+	#define UE_ENABLE_ACTOR_REPLIST_TRACKING 0
+#endif
+
+#define UE_ACTOR_REPLIST_TYPE_EXTRA_SAFETY			(UE_SERVER && UE_ENABLE_ACTOR_REPLIST_TRACKING)
+
+#if !UE_ACTOR_REPLIST_TYPE_EXTRA_SAFETY
 typedef AActor* FActorRepListType;
+FORCEINLINE bool DoesActorPointerLookValid(const AActor* In)
+{
+	return ((uint64)(In) & 0x0F) == 0;
+}
+#else
+struct FActorRepListType
+{
+	/** Actual load-bearing payload */
+	AActor* ActorRaw;
+	/** Validity test */
+	FObjectKey ActorKey;
+
+	/** More info for debugging - ActorRaw's name */
+	FName ActorName;
+	/** More info for debugging - ActorRaw's Owner name */
+	FName OwnerName;
+	/** More info for debugging - ActorRaw's Outer's package name */
+	FName OuterPackageName;
+
+	inline void SetDebugInfo()
+	{
+		ActorName = OwnerName = OuterPackageName = NAME_None;
+		if (LIKELY(ActorRaw))
+		{
+			ActorName = ActorRaw->GetFName();
+			OwnerName = ActorRaw->GetOwner() ? ActorRaw->GetOwner()->GetFName() : NAME_None;
+			if (LIKELY(ActorRaw->GetOuter()))
+			{
+				// judging by the implementation GetPackage() cannot return nullptr, but play it safe
+				UPackage* Pkg = ActorRaw->GetOuter()->GetPackage();
+				if (LIKELY(Pkg))
+				{
+					OuterPackageName = Pkg->GetFName();
+				}
+			}
+		}
+	}
+
+	FActorRepListType() = default;
+
+	FActorRepListType(AActor* InActor)
+		: ActorRaw(InActor)
+		, ActorKey(InActor)
+	{
+		SetDebugInfo();
+	}
+
+	// to support conversion from TObjectPtr<ASubclassOfActor>
+	template <
+		typename T,
+		decltype(ImplicitConv<AActor*>(std::declval<const T&>())) = nullptr
+	>
+	FActorRepListType(const T& InActor)
+		: ActorRaw(InActor)
+		, ActorKey(InActor)
+	{
+		SetDebugInfo();
+	}
+
+	operator AActor*() { return ActorRaw; }
+	operator AActor*() const { return ActorRaw; }
+	AActor* operator->() { return ActorRaw; }
+	AActor* operator->() const { return ActorRaw; }
+	operator uint64() const { return reinterpret_cast<uint64>(ActorRaw); }
+	FActorRepListType& operator=(FActorRepListType const& InActor) = default;
+	FActorRepListType& operator=(AActor* InActor)
+	{
+		ActorRaw = InActor;
+		ActorKey = InActor;
+		SetDebugInfo();
+		return *this;
+	}
+	FActorRepListType& operator=(TObjectPtr<AActor> InActor)
+	{
+		ActorRaw = InActor;
+		ActorKey = InActor;
+		SetDebugInfo();
+		return *this;
+	}
+	bool operator==(FActorRepListType const& Other) const
+	{
+		return ActorRaw == Other.ActorRaw;
+	}
+	bool operator==(AActor* Other) const
+	{
+		return ActorRaw == Other;
+	}
+	bool IsValid() const
+	{
+		UObject const* Object = ActorKey.ResolveObjectPtr();
+		return Object && Object == static_cast<UObject const*>(ActorRaw);
+	}
+};
+template< class T > FORCEINLINE T* Cast(const FActorRepListType& Src) { return Cast<T>(Src.ActorRaw); }
+template< class T > FORCEINLINE T* ExactCast(const FActorRepListType& Src) { return ExactCast<T>(Src.ActorRaw); }
+template< class T > FORCEINLINE T* CastChecked(const FActorRepListType& Src, ECastCheckedType::Type CheckType = ECastCheckedType::NullChecked) { return CastChecked<T>(Src.ActorRaw, CheckType); }
+
+
+FORCEINLINE bool DoesActorPointerLookValid(const FActorRepListType& In)
+{
+	return In.IsValid();
+}
+#endif
+
 FORCEINLINE FString GetActorRepListTypeDebugString(const FActorRepListType& In) { return GetNameSafe(In); }
 FORCEINLINE UClass* GetActorRepListTypeClass(const FActorRepListType& In) { return In->GetClass(); }
 
@@ -80,7 +192,21 @@ enum class EActorRepListTypeFlags : uint8
 };
 
 // Tests if an actor is valid for replication: not pending kill, etc. Says nothing about wanting to replicate or should replicate, etc.
-FORCEINLINE bool IsActorValidForReplication(const FActorRepListType& In) { return !In->IsActorBeingDestroyed() && IsValidChecked(In) && !In->IsUnreachable(); }
+FORCEINLINE bool IsActorValidForReplication(const FActorRepListType& In)
+{ 
+	return DoesActorPointerLookValid(In) && !In->IsActorBeingDestroyed() && IsValidChecked(In) && !In->IsUnreachable(); 
+}
+REPLICATIONGRAPH_API void LogMoreInfoOnIsActorValidFailure(const FActorRepListType& In);
+FORCEINLINE bool IsActorValidForReplication_LogMoreInfo(const FActorRepListType& In)
+{ 
+	if (LIKELY(IsActorValidForReplication(In)))
+	{
+		return true;
+	}
+
+	LogMoreInfoOnIsActorValidFailure(In);
+	return false;
+}
 
 // Tests if an actor is valid for replication gathering. Meaning, it can be gathered from the replication graph and considered for replication.
 FORCEINLINE bool IsActorValidForReplicationGather(const FActorRepListType& In)
@@ -921,7 +1047,7 @@ struct FGlobalActorReplicationInfoMap
 			return *Ptr->Get();
 		}
 
-		ensureMsgf(IsActorValidForReplication(Actor), TEXT("This obj %s is pending to kill, storing this data will generate stale data in the map."), *GetPathNameSafe(Actor));
+		ensureMsgf(IsActorValidForReplication_LogMoreInfo(Actor), TEXT("An invalid actor pointer is passed to FGlobalActorReplicationInfo::Get(), storing this data will generate stale data in the map."));
 
 		// We need to add data for this actor
 		FClassReplicationInfo& ClassInfo = GetClassInfo( GetActorRepListTypeClass(Actor) );
@@ -941,7 +1067,7 @@ struct FGlobalActorReplicationInfoMap
 			return *Ptr->Get();
 		}
 
-		ensureMsgf(IsActorValidForReplication(Actor), TEXT("This obj %s is pending to kill, storing this data will generate stale data in the map."), *GetPathNameSafe(Actor));
+		ensureMsgf(IsActorValidForReplication_LogMoreInfo(Actor), TEXT("An invalid actor pointer is passed to FGlobalActorReplicationInfo::Get(), storing this data will generate stale data in the map."));
 
 		bWasCreated = true;
 
