Changed UNetDriver::ERemoteFunctionSendPolicy from an enum to an enum class so that other instances of "Default" won't conflict with it.
This class isn't at fault, but it's the easiest change.

diff --git a/Engine/Plugins/Runtime/ReplicationGraph/Source/Private/ReplicationGraph.cpp b/Engine/Plugins/Runtime/ReplicationGraph/Source/Private/ReplicationGraph.cpp
--- a/Engine/Plugins/Runtime/ReplicationGraph/Source/Private/ReplicationGraph.cpp
+++ b/Engine/Plugins/Runtime/ReplicationGraph/Source/Private/ReplicationGraph.cpp
@@ -2345,14 +2345,14 @@ bool UReplicationGraph::ProcessRemoteFunction(class AActor* Actor, UFunction* Fu
 
 		TOptional<FVector> ActorLocation;
 
-		UNetDriver::ERemoteFunctionSendPolicy SendPolicy = UNetDriver::Default;
+		UNetDriver::ERemoteFunctionSendPolicy SendPolicy = UNetDriver::ERemoteFunctionSendPolicy::Default;
 		if (CVar_RepGraph_EnableRPCSendPolicy > 0)
 		{
 			if (FRPCSendPolicyInfo* FuncSendPolicy = RPCSendPolicyMap.Find(FObjectKey(Function)))
 			{
 				if (FuncSendPolicy->bSendImmediately)
 				{
-					SendPolicy = UNetDriver::ForceSend;
+					SendPolicy = UNetDriver::ERemoteFunctionSendPolicy::ForceSend;
 				}
 			}
 		}
@@ -2464,7 +2464,7 @@ bool UReplicationGraph::ProcessRemoteFunction(class AActor* Actor, UFunction* Fu
 			{
 				NetDriver->ProcessRemoteFunctionForChannel(ConnectionActorInfo.Channel, ClassCache, FieldCache, TargetObj, NetConnection, Function, Parameters, OutParms, Stack, true, SendPolicy, RemoteFunctionFlags);
 
-				if (SendPolicy == UNetDriver::ForceSend)
+				if (SendPolicy == UNetDriver::ERemoteFunctionSendPolicy::ForceSend)
 				{
 					// Queue the send in an array that we consume in PostTickDispatch to avoid force flushing multiple times a frame on the same connection
 					ConnectionsNeedingsPostTickDispatchFlush.AddUnique(NetConnection);
diff --git a/Engine/Source/Runtime/Engine/Classes/Engine/NetDriver.h b/Engine/Source/Runtime/Engine/Classes/Engine/NetDriver.h
--- a/Engine/Source/Runtime/Engine/Classes/Engine/NetDriver.h
+++ b/Engine/Source/Runtime/Engine/Classes/Engine/NetDriver.h
@@ -1530,7 +1530,7 @@ public:
 	/** Return a reference to the database that holds metrics calcluated by the networking system. */
 	ENGINE_API TObjectPtr<UNetworkMetricsDatabase> GetMetrics() { return NetworkMetricsDatabase; };
 
-	enum ERemoteFunctionSendPolicy
+	enum class ERemoteFunctionSendPolicy
 	{		
 		/** Unreliable multicast are queued. Everything else is send immediately */
 		Default, 
@@ -1540,7 +1540,13 @@ public:
 
 		/** Bunch is queued until next actor replication, no matter what */
 		ForceQueue,
-	};	
+	};
+	UE_DEPRECATED(5.4, "Use fully scoped enum class value UNetDriver::ERemoteFunctionSendPolicy::Default")
+	static constexpr ERemoteFunctionSendPolicy Default = ERemoteFunctionSendPolicy::Default;
+	UE_DEPRECATED(5.4, "Use fully scoped enum class value UNetDriver::ERemoteFunctionSendPolicy::ForceSend")
+	static constexpr ERemoteFunctionSendPolicy ForceSend = ERemoteFunctionSendPolicy::ForceSend;
+	UE_DEPRECATED(5.4, "Use fully scoped enum class value UNetDriver::ERemoteFunctionSendPolicy::ForceQueue")
+	static constexpr ERemoteFunctionSendPolicy ForceQueue = ERemoteFunctionSendPolicy::ForceQueue;
 
 	/** Process a remote function on given actor channel. This is called by ::ProcessRemoteFunction.*/
 	ENGINE_API void ProcessRemoteFunctionForChannel(
