RepGraph: Fix skipping subsequent connections if one connection had nothing to replicate

diff --git a/Engine/Plugins/Runtime/ReplicationGraph/Source/Private/ReplicationGraph.cpp b/Engine/Plugins/Runtime/ReplicationGraph/Source/Private/ReplicationGraph.cpp
--- a/Engine/Plugins/Runtime/ReplicationGraph/Source/Private/ReplicationGraph.cpp
+++ b/Engine/Plugins/Runtime/ReplicationGraph/Source/Private/ReplicationGraph.cpp
@@ -1226,9 +1226,8 @@ int32 UReplicationGraph::ServerReplicateActors(float DeltaSeconds)
 
 			if (GatheredReplicationListsForConnection.NumLists() == 0)
 			{
-				// No lists were returned, kind of weird but not fatal. Early out because code below assumes at least 1 list
-				UE_LOG(LogReplicationGraph, Warning, TEXT("No Replication Lists were returned for connection"));
-				return 0;
+				UE_LOG(LogReplicationGraph, Verbose, TEXT("No Replication Lists were returned for connection"));
+				continue;
 			}
 
 			for( const FNetViewer& NetViewer : ConnectionViewers )
