RepGraph: Hold gathered actors in owned TArray rather than views of unowned actor lists
#tests 2 clients on windows

diff --git a/Engine/Plugins/Runtime/ReplicationGraph/Source/Private/ReplicationGraph.cpp b/Engine/Plugins/Runtime/ReplicationGraph/Source/Private/ReplicationGraph.cpp
--- a/Engine/Plugins/Runtime/ReplicationGraph/Source/Private/ReplicationGraph.cpp
+++ b/Engine/Plugins/Runtime/ReplicationGraph/Source/Private/ReplicationGraph.cpp
@@ -1395,7 +1395,6 @@ void UReplicationGraph::ReplicateActorListsForConnections_Default(UNetReplicatio
 	FNativeClassAccumulator DormancyClassAccumulator;
 	FNativeClassAccumulator DistanceClassAccumulator;
 
-	int32 NumGatheredListsOnConnection = 0;
 	int32 NumGatheredActorsOnConnection = 0;
 	int32 NumPrioritizedActorsOnConnection = 0;
 
@@ -1415,192 +1414,188 @@ void UReplicationGraph::ReplicateActorListsForConnections_Default(UNetReplicatio
 		PrioritizedReplicationList.Reset();
 		TArray<FPrioritizedRepList::FItem>* SortingArray = &PrioritizedReplicationList.Items;
 
-		NumGatheredListsOnConnection += GatheredReplicationListsForConnection.NumLists();
-
 		const float MaxDistanceScaling = PrioritizationConstants.MaxDistanceScaling;
 		const uint32 MaxFramesSinceLastRep = PrioritizationConstants.MaxFramesSinceLastRep;
 
-		const TArray<FActorRepListConstView>& GatheredLists = GatheredReplicationListsForConnection.GetLists(EActorRepListTypeFlags::Default);
-		for (const FActorRepListConstView& List : GatheredLists)
+		// Add actors from gathered list
+		const TArrayView<const FActorRepListType> Actors = GatheredReplicationListsForConnection.ViewActors(EActorRepListTypeFlags::Default);
+		NumGatheredActorsOnConnection += Actors.Num();
+
+		for (const FActorRepListType& Actor : Actors)
 		{
-			// Add actors from gathered list
-			NumGatheredActorsOnConnection += List.Num();
-			for (const FActorRepListType& Actor : List)
-			{
-				RG_QUICK_SCOPE_CYCLE_COUNTER(Prioritize_InnerLoop);
+			RG_QUICK_SCOPE_CYCLE_COUNTER(Prioritize_InnerLoop);
 
-				if (!ensureMsgf(IsActorValidForReplication_LogMoreInfo(Actor), TEXT("Actor not valid for replication")))
-				{
-					continue;
-				}
+			if (!ensureMsgf(IsActorValidForReplication_LogMoreInfo(Actor), TEXT("Actor not valid for replication")))
+			{
+				continue;
+			}
 
-				// -----------------------------------------------------------------------------------------------------------------
-				//	Prioritize Actor for Connection: this is the main block of code for calculating a final score for this actor
-				//		-This is still pretty rough. It would be nice if this was customizable per project without suffering virtual calls.
-				// -----------------------------------------------------------------------------------------------------------------
+			// -----------------------------------------------------------------------------------------------------------------
+			//	Prioritize Actor for Connection: this is the main block of code for calculating a final score for this actor
+			//		-This is still pretty rough. It would be nice if this was customizable per project without suffering virtual calls.
+			// -----------------------------------------------------------------------------------------------------------------
 
-				if (RepGraphConditionalActorBreakpoint(Actor, NetConnection))
-				{
-					UE_LOG(LogReplicationGraph, Display, TEXT("UReplicationGraph PrioritizeActor: %s"), *Actor->GetName());
-				}
+			if (RepGraphConditionalActorBreakpoint(Actor, NetConnection))
+			{
+				UE_LOG(LogReplicationGraph, Display, TEXT("UReplicationGraph PrioritizeActor: %s"), *Actor->GetName());
+			}
 
-				FConnectionReplicationActorInfo& ConnectionData = ConnectionActorInfoMap.FindOrAdd(Actor);
+			FConnectionReplicationActorInfo& ConnectionData = ConnectionActorInfoMap.FindOrAdd(Actor);
 
-				RG_QUICK_SCOPE_CYCLE_COUNTER(Prioritize_InnerLoop_ConnGlobalLookUp);
+			RG_QUICK_SCOPE_CYCLE_COUNTER(Prioritize_InnerLoop_ConnGlobalLookUp);
 
-				// Skip if dormant on this connection. We want this to always be the first/quickest check.
-				if (ConnectionData.bDormantOnConnection)
+			// Skip if dormant on this connection. We want this to always be the first/quickest check.
+			if (ConnectionData.bDormantOnConnection)
+			{
+				DO_REPGRAPH_DETAILS(PrioritizedReplicationList.GetNextSkippedDebugDetails(Actor)->bWasDormant = true);
+				if (bDoCulledOnConnectionCount)
 				{
-					DO_REPGRAPH_DETAILS(PrioritizedReplicationList.GetNextSkippedDebugDetails(Actor)->bWasDormant = true);
-					if (bDoCulledOnConnectionCount)
-					{
-						DormancyClassAccumulator.Increment(Actor->GetClass());
-					}
-					continue;
+					DormancyClassAccumulator.Increment(Actor->GetClass());
 				}
+				continue;
+			}
 
-				FGlobalActorReplicationInfo& GlobalData = GlobalActorReplicationInfoMap.Get(Actor);
+			FGlobalActorReplicationInfo& GlobalData = GlobalActorReplicationInfoMap.Get(Actor);
 
-				RG_QUICK_SCOPE_CYCLE_COUNTER(Prioritize_InnerLoop_PostGlobalLookUp);
+			RG_QUICK_SCOPE_CYCLE_COUNTER(Prioritize_InnerLoop_PostGlobalLookUp);
 
-				// Skip if its not time to replicate on this connection yet. We have to look at ForceNetUpdateFrame here. It would be possible to clear
-				// NextReplicationFrameNum on all connections when ForceNetUpdate is called. This probably means more work overall per frame though. Something to consider.
-				if (!ReadyForNextReplication(ConnectionData, GlobalData, FrameNum))
-				{
-					DO_REPGRAPH_DETAILS(PrioritizedReplicationList.GetNextSkippedDebugDetails(Actor)->FramesTillNextReplication = (FrameNum - ConnectionData.LastRepFrameNum));
-					continue;
-				}
+			// Skip if its not time to replicate on this connection yet. We have to look at ForceNetUpdateFrame here. It would be possible to clear
+			// NextReplicationFrameNum on all connections when ForceNetUpdate is called. This probably means more work overall per frame though. Something to consider.
+			if (!ReadyForNextReplication(ConnectionData, GlobalData, FrameNum))
+			{
+				DO_REPGRAPH_DETAILS(PrioritizedReplicationList.GetNextSkippedDebugDetails(Actor)->FramesTillNextReplication = (FrameNum - ConnectionData.LastRepFrameNum));
+				continue;
+			}
 
-				RG_QUICK_SCOPE_CYCLE_COUNTER(Prioritize_InnerLoop_PostReady);
+			RG_QUICK_SCOPE_CYCLE_COUNTER(Prioritize_InnerLoop_PostReady);
 
-				// Output record for full debugging. This is not used in the actual sorting/prioritization of the list, just for logging/debugging purposes
-				FPrioritizedActorFullDebugDetails* DebugDetails = nullptr;
-				if (DO_REPGRAPH_DETAILS(UNLIKELY(bEnableFullActorPrioritizationDetails)))
-				{
-					DO_REPGRAPH_DETAILS(DebugDetails = PrioritizedReplicationList.GetNextFullDebugDetails(Actor));
-				}
+			// Output record for full debugging. This is not used in the actual sorting/prioritization of the list, just for logging/debugging purposes
+			FPrioritizedActorFullDebugDetails* DebugDetails = nullptr;
+			if (DO_REPGRAPH_DETAILS(UNLIKELY(bEnableFullActorPrioritizationDetails)))
+			{
+				DO_REPGRAPH_DETAILS(DebugDetails = PrioritizedReplicationList.GetNextFullDebugDetails(Actor));
+			}
 
-				float AccumulatedPriority = GlobalData.Settings.AccumulatedNetPriorityBias;
+			float AccumulatedPriority = GlobalData.Settings.AccumulatedNetPriorityBias;
 
-				// -------------------
-				// Distance Scaling
-				// -------------------
-				if (GlobalData.Settings.DistancePriorityScale > 0.f)
+			// -------------------
+			// Distance Scaling
+			// -------------------
+			if (GlobalData.Settings.DistancePriorityScale > 0.f)
+			{
+				// Always compute distance even for AlwaysRelevant actors since the priority scaling needs it.
+				FVector::FReal SmallestDistanceSq = std::numeric_limits<FVector::FReal>::max();
+				int32 ViewersThatSkipActor = 0;
+				
+				for (const FNetViewer& CurViewer : Viewers)
 				{
-					// Always compute distance even for AlwaysRelevant actors since the priority scaling needs it.
-					FVector::FReal SmallestDistanceSq = std::numeric_limits<FVector::FReal>::max();
-					int32 ViewersThatSkipActor = 0;
-					
-					for (const FNetViewer& CurViewer : Viewers)
-					{
-						const FVector::FReal DistSq = (GlobalData.WorldLocation - CurViewer.ViewLocation).SizeSquared();
-						SmallestDistanceSq = FMath::Min(DistSq, SmallestDistanceSq);
-
-						// Figure out if we should be skipping this actor
-						if (bDoDistanceCull && ConnectionData.GetCullDistanceSquared() > 0.f && DistSq > ConnectionData.GetCullDistanceSquared())
-						{
-							++ViewersThatSkipActor;
-							continue;
-						}
-					}
+					const FVector::FReal DistSq = (GlobalData.WorldLocation - CurViewer.ViewLocation).SizeSquared();
+					SmallestDistanceSq = FMath::Min(DistSq, SmallestDistanceSq);
 
-					// If no one is near this actor, skip it.
-					if (ViewersThatSkipActor >= Viewers.Num())
+					// Figure out if we should be skipping this actor
+					if (bDoDistanceCull && ConnectionData.GetCullDistanceSquared() > 0.f && DistSq > ConnectionData.GetCullDistanceSquared())
 					{
-						DO_REPGRAPH_DETAILS(PrioritizedReplicationList.GetNextSkippedDebugDetails(Actor)->DistanceCulled = static_cast<float>(FMath::Sqrt(SmallestDistanceSq)));
-
-						// Skipped actors should not have any 
-						if (bDoCulledOnConnectionCount)
-						{
-							DistanceClassAccumulator.Increment(Actor->GetClass());
-						}
+						++ViewersThatSkipActor;
 						continue;
 					}
+				}
+
+				// If no one is near this actor, skip it.
+				if (ViewersThatSkipActor >= Viewers.Num())
+				{
+					DO_REPGRAPH_DETAILS(PrioritizedReplicationList.GetNextSkippedDebugDetails(Actor)->DistanceCulled = static_cast<float>(FMath::Sqrt(SmallestDistanceSq)));
 
-					const float DistanceFactor = FMath::Clamp(static_cast<float>(SmallestDistanceSq / MaxDistanceScaling), 0.f, 1.f) * GlobalData.Settings.DistancePriorityScale;
-					if (DO_REPGRAPH_DETAILS(UNLIKELY(DebugDetails)))
+					// Skipped actors should not have any 
+					if (bDoCulledOnConnectionCount)
 					{
-						DebugDetails->DistanceSq = SmallestDistanceSq;
-						DebugDetails->DistanceFactor = DistanceFactor;
+						DistanceClassAccumulator.Increment(Actor->GetClass());
 					}
-
-					AccumulatedPriority += DistanceFactor;
+					continue;
 				}
 
-				RG_QUICK_SCOPE_CYCLE_COUNTER(Prioritize_InnerLoop_PostCull);
+				const float DistanceFactor = FMath::Clamp(static_cast<float>(SmallestDistanceSq / MaxDistanceScaling), 0.f, 1.f) * GlobalData.Settings.DistancePriorityScale;
+				if (DO_REPGRAPH_DETAILS(UNLIKELY(DebugDetails)))
+				{
+					DebugDetails->DistanceSq = SmallestDistanceSq;
+					DebugDetails->DistanceFactor = DistanceFactor;
+				}
 
-				// Update the timeout frame number here. (Since this was returned by the graph, regardless if we end up replicating or not, we bump up the timeout frame num. This has to be done here because Distance Scaling can cull the actor
-				UpdateActorChannelCloseFrameNum(Actor, ConnectionData, GlobalData, FrameNum, NetConnection);
+				AccumulatedPriority += DistanceFactor;
+			}
 
-				//UE_CLOG(DebugConnection, LogReplicationGraph, Display, TEXT("0x%X0x%X ConnectionData.ActorChannelCloseFrameNum=%d on %d"), Actor, NetConnection, ConnectionData.ActorChannelCloseFrameNum, FrameNum);
+			RG_QUICK_SCOPE_CYCLE_COUNTER(Prioritize_InnerLoop_PostCull);
 
-				// -------------------
-				// Starvation Scaling
-				// -------------------
-				if (GlobalData.Settings.StarvationPriorityScale > 0.f)
-				{
-					// StarvationPriorityScale = scale "Frames since last rep". E.g, 2.0 means treat every missed frame as if it were 2, etc.
-					const float FramesSinceLastRep = ((float)(FrameNum - ConnectionData.LastRepFrameNum)) * GlobalData.Settings.StarvationPriorityScale;
-					const float StarvationFactor = 1.f - FMath::Clamp<float>(FramesSinceLastRep / (float)MaxFramesSinceLastRep, 0.f, 1.f);
+			// Update the timeout frame number here. (Since this was returned by the graph, regardless if we end up replicating or not, we bump up the timeout frame num. This has to be done here because Distance Scaling can cull the actor
+			UpdateActorChannelCloseFrameNum(Actor, ConnectionData, GlobalData, FrameNum, NetConnection);
 
-					AccumulatedPriority += StarvationFactor;
+			//UE_CLOG(DebugConnection, LogReplicationGraph, Display, TEXT("0x%X0x%X ConnectionData.ActorChannelCloseFrameNum=%d on %d"), Actor, NetConnection, ConnectionData.ActorChannelCloseFrameNum, FrameNum);
 
-					if (DO_REPGRAPH_DETAILS(UNLIKELY(DebugDetails)))
-					{
-						DebugDetails->FramesSinceLastRap = static_cast<uint32>(FMath::TruncToInt32(FramesSinceLastRep));
-						DebugDetails->StarvationFactor = StarvationFactor;
-					}
-				}
+			// -------------------
+			// Starvation Scaling
+			// -------------------
+			if (GlobalData.Settings.StarvationPriorityScale > 0.f)
+			{
+				// StarvationPriorityScale = scale "Frames since last rep". E.g, 2.0 means treat every missed frame as if it were 2, etc.
+				const float FramesSinceLastRep = ((float)(FrameNum - ConnectionData.LastRepFrameNum)) * GlobalData.Settings.StarvationPriorityScale;
+				const float StarvationFactor = 1.f - FMath::Clamp<float>(FramesSinceLastRep / (float)MaxFramesSinceLastRep, 0.f, 1.f);
 
-				// ------------------------
-				// Pending dormancy scaling
-				// ------------------------
+				AccumulatedPriority += StarvationFactor;
 
-				// Make sure pending dormant actors that have replicated at least once are prioritized,
-				// so we actually mark them dormant quickly, skip future work, and close their channels.
-				// Otherwise, newly spawned or never-replicated actors may starve out existing actors trying to go dormant.
-				if (GlobalData.bWantsToBeDormant && ConnectionData.LastRepFrameNum > 0)
+				if (DO_REPGRAPH_DETAILS(UNLIKELY(DebugDetails)))
 				{
-					AccumulatedPriority -= 1.5f;
+					DebugDetails->FramesSinceLastRap = static_cast<uint32>(FMath::TruncToInt32(FramesSinceLastRep));
+					DebugDetails->StarvationFactor = StarvationFactor;
 				}
+			}
 
-				// -------------------
-				//	Game code priority
-				// -------------------
+			// ------------------------
+			// Pending dormancy scaling
+			// ------------------------
 
-				if (GlobalData.ForceNetUpdateFrame > ConnectionData.LastRepFrameNum)
-				{
-					// Note that in legacy ForceNetUpdate did not actually bump priority. This gives us a hard coded bump if we haven't replicated since the last ForceNetUpdate frame.
-					AccumulatedPriority -= 1.f;
+			// Make sure pending dormant actors that have replicated at least once are prioritized,
+			// so we actually mark them dormant quickly, skip future work, and close their channels.
+			// Otherwise, newly spawned or never-replicated actors may starve out existing actors trying to go dormant.
+			if (GlobalData.bWantsToBeDormant && ConnectionData.LastRepFrameNum > 0)
+			{
+				AccumulatedPriority -= 1.5f;
+			}
+
+			// -------------------
+			//	Game code priority
+			// -------------------
+
+			if (GlobalData.ForceNetUpdateFrame > ConnectionData.LastRepFrameNum)
+			{
+				// Note that in legacy ForceNetUpdate did not actually bump priority. This gives us a hard coded bump if we haven't replicated since the last ForceNetUpdate frame.
+				AccumulatedPriority -= 1.f;
 
-					if (DO_REPGRAPH_DETAILS(UNLIKELY(DebugDetails)))
+				if (DO_REPGRAPH_DETAILS(UNLIKELY(DebugDetails)))
+				{
+					DebugDetails->GameCodeScaling = -1.f;
+				}
+			}
+			
+			// -------------------
+			// Always prioritize the connection's owner and view target, since these are the most important actors for the client.
+			// -------------------
+			for (const FNetViewer& CurViewer : Viewers)
+			{
+				// We need to find if this is anyone's viewer or viewtarget, not just the parent connection.
+				if (Actor == CurViewer.ViewTarget || Actor == CurViewer.InViewer)
+				{
+					if (UE::Net::CVar_ForceConnectionViewerPriority > 0)
 					{
-						DebugDetails->GameCodeScaling = -1.f;
+						AccumulatedPriority = -MAX_FLT;
 					}
-                }
-				
-				// -------------------
-				// Always prioritize the connection's owner and view target, since these are the most important actors for the client.
-				// -------------------
-				for (const FNetViewer& CurViewer : Viewers)
-				{
-					// We need to find if this is anyone's viewer or viewtarget, not just the parent connection.
-					if (Actor == CurViewer.ViewTarget || Actor == CurViewer.InViewer)
+					else
 					{
-						if (UE::Net::CVar_ForceConnectionViewerPriority > 0)
-						{
-							AccumulatedPriority = -MAX_FLT;
-						}
-						else
-						{
-							AccumulatedPriority -= 10.0f;
-						}
-						break;
+						AccumulatedPriority -= 10.0f;
 					}
+					break;
 				}
-
-				SortingArray->Emplace(FPrioritizedRepList::FItem(AccumulatedPriority, Actor, &GlobalData, &ConnectionData));
 			}
+
+			SortingArray->Emplace(FPrioritizedRepList::FItem(AccumulatedPriority, Actor, &GlobalData, &ConnectionData));
 		}
 
 		{
@@ -1627,7 +1622,7 @@ void UReplicationGraph::ReplicateActorListsForConnections_Default(UNetReplicatio
 		UE_LOG(LogReplicationGraph, Display, TEXT("Saturated Connections: %d"), GNumSaturatedConnections);
 		UE_LOG(LogReplicationGraph, Display, TEXT(""));
 
-		UE_LOG(LogReplicationGraph, Display, TEXT("Gathered Lists: %d Gathered Actors: %d  PrioritizedActors: %d"), NumGatheredListsOnConnection, NumGatheredActorsOnConnection, NumPrioritizedActorsOnConnection);
+		UE_LOG(LogReplicationGraph, Display, TEXT("Gathered Actors: %d, PrioritizedActors: %d"), NumGatheredActorsOnConnection, NumPrioritizedActorsOnConnection);
 		UE_LOG(LogReplicationGraph, Display, TEXT("Connection Loaded Streaming Levels: %d"), NetConnection->ClientVisibleLevelNames.Num());
 	}
 #endif // WITH_SERVER_CODE
@@ -1747,95 +1742,91 @@ void UReplicationGraph::ReplicateActorListsForConnections_FastShared(UNetReplica
 	// This really isn't ideal. We want to have better ways of tracking and limiting network traffic. This feels pretty hacky in implementation but conceptually is good.
 	FScopedQueuedBits ScopedQueuedBits(NetConnection->QueuedBits, TotalBitsWritten);
 
-	const TArray<FActorRepListConstView>& GatheredLists = GatheredReplicationListsForConnection.GetLists(EActorRepListTypeFlags::FastShared);
-	for (int32 ListIdx = 0; ListIdx < GatheredLists.Num(); ++ListIdx)
+	const TArrayView<const FActorRepListType> Actors = GatheredReplicationListsForConnection.ViewActors(EActorRepListTypeFlags::FastShared);
+	for (int32 i = 0; i < Actors.Num(); ++i)
 	{
-		const FActorRepListConstView& List = GatheredLists[(ListIdx + FrameNum) % GatheredLists.Num()];
-		for (int32 i = 0; i < List.Num(); ++i)
+		// Round robin through the list over multiple frames. We want to avoid sorting this list based on 'time since last rep'. This is a good balance
+		AActor* Actor = Actors[(StartIdx + i) % Actors.Num()];
+
+		int32 BitsWritten = 0;
+
+		if (RepGraphConditionalActorBreakpoint(Actor, NetConnection))
 		{
-			// Round robin through the list over multiple frames. We want to avoid sorting this list based on 'time since last rep'. This is a good balance
-			AActor* Actor = List[(i + StartIdx) % List.Num()];
+			UE_LOG(LogReplicationGraph, Display, TEXT("UReplicationGraph FastShared Path Replication: %s"), *Actor->GetName());
+		}
 
-			int32 BitsWritten = 0;
+		FConnectionReplicationActorInfo& ConnectionData = ConnectionActorInfoMap.FindOrAdd(Actor);
 
-			if (RepGraphConditionalActorBreakpoint(Actor, NetConnection))
-			{
-				UE_LOG(LogReplicationGraph, Display, TEXT("UReplicationGraph FastShared Path Replication: %s"), *Actor->GetName());
-			}
+		// Don't fast path rep if we already repped in the default path this frame
+		if (UNLIKELY(ConnectionData.LastRepFrameNum == FrameNum))
+		{
+			continue;
+		}
 
-			FConnectionReplicationActorInfo& ConnectionData = ConnectionActorInfoMap.FindOrAdd(Actor);
+		if (UNLIKELY(ConnectionData.bTearOff))
+		{
+			continue;
+		}
 
-			// Don't fast path rep if we already repped in the default path this frame
-			if (UNLIKELY(ConnectionData.LastRepFrameNum == FrameNum))
-			{
-				continue;
-			}
+		// Actor channel must already be established to rep fast path
+		UActorChannel* ActorChannel = ConnectionData.Channel;
+		if (ActorChannel == nullptr || ActorChannel->Closing)
+		{
+			continue;
+		}
 
-			if (UNLIKELY(ConnectionData.bTearOff))
-			{
-				continue;
-			}
+		FGlobalActorReplicationInfo& GlobalActorInfo = GlobalActorReplicationInfoMap.Get(Actor);
+		if (GlobalActorInfo.Settings.FastSharedReplicationFunc == nullptr)
+		{
+			// This actor does not support fastshared replication
+			// FIXME: we should avoid this by keeping these actors on separate lists
+			continue;
+		}
 
-			// Actor channel must already be established to rep fast path
-			UActorChannel* ActorChannel = ConnectionData.Channel;
-			if (ActorChannel == nullptr || ActorChannel->Closing)
-			{
-				continue;
-			}
+		// Determine if this actor has any view relevancy to any connection this client has
+		bool bNoViewRelevency = true;
+		for (const FNetViewer& CurView : Viewers)
+		{
+			const FVector& ConnectionViewLocation = CurView.ViewLocation;
+			const FVector& ConnectionViewDir = CurView.ViewDir;
 
-			FGlobalActorReplicationInfo& GlobalActorInfo = GlobalActorReplicationInfoMap.Get(Actor);
-			if (GlobalActorInfo.Settings.FastSharedReplicationFunc == nullptr)
+			// Simple dot product rejection: only fast rep actors in front of this connection
+			const FVector DirToActor = GlobalActorInfo.WorldLocation - ConnectionViewLocation;
+			if (!(FVector::DotProduct(DirToActor, ConnectionViewDir) < 0.f))
 			{
-				// This actor does not support fastshared replication
-				// FIXME: we should avoid this by keeping these actors on separate lists
-				continue;
+				bNoViewRelevency = false;
+				break;
 			}
 
-			// Determine if this actor has any view relevancy to any connection this client has
-			bool bNoViewRelevency = true;
-			for (const FNetViewer& CurView : Viewers)
+			// Simple distance cull
+			const FVector::FReal DistSq = DirToActor.SizeSquared();
+			if (!(DistSq > (ConnectionData.GetCullDistanceSquared() * FastSharedDistanceRequirementPct)))
 			{
-				const FVector& ConnectionViewLocation = CurView.ViewLocation;
-				const FVector& ConnectionViewDir = CurView.ViewDir;
-
-				// Simple dot product rejection: only fast rep actors in front of this connection
-				const FVector DirToActor = GlobalActorInfo.WorldLocation - ConnectionViewLocation;
-				if (!(FVector::DotProduct(DirToActor, ConnectionViewDir) < 0.f))
-				{
-					bNoViewRelevency = false;
-					break;
-				}
-
-				// Simple distance cull
-				const FVector::FReal DistSq = DirToActor.SizeSquared();
-				if (!(DistSq > (ConnectionData.GetCullDistanceSquared() * FastSharedDistanceRequirementPct)))
-				{
-					bNoViewRelevency = false;
-					break;
-				}
+				bNoViewRelevency = false;
+				break;
 			}
+		}
 
-			// Skip out if they have none.
-			if (bNoViewRelevency)
-			{
-				continue;
-			}
+		// Skip out if they have none.
+		if (bNoViewRelevency)
+		{
+			continue;
+		}
 
-			BitsWritten = (int32)ReplicateSingleActor_FastShared(Actor, ConnectionData, GlobalActorInfo, *ConnectionManager, FrameNum);
-			TotalBitsWritten += BitsWritten;
+		BitsWritten = (int32)ReplicateSingleActor_FastShared(Actor, ConnectionData, GlobalActorInfo, *ConnectionManager, FrameNum);
+		TotalBitsWritten += BitsWritten;
 
 #if !(UE_BUILD_SHIPPING || UE_BUILD_TEST)
-			static bool SkipCheck = false;
-			if (SkipCheck)
-			{
-				continue;
-			}
+		static bool SkipCheck = false;
+		if (SkipCheck)
+		{
+			continue;
+		}
 #endif
-			if (TotalBitsWritten > MaxBits)
-			{
-				NotifyConnectionFastPathSaturated();
-				return;
-			}
+		if (TotalBitsWritten > MaxBits)
+		{
+			NotifyConnectionFastPathSaturated();
+			return;
 		}
 	}
 #endif // WITH_SERVER_CODE
@@ -2130,38 +2121,35 @@ int64 UReplicationGraph::ReplicateSingleActor(const FActorRepListType& Actor, FC
 
 		const int32 CloseFrameNum = ActorInfo.ActorChannelCloseFrameNum;
 		
-		const TArray<FActorRepListConstView>& DependentLists = ListGatherer.GetLists(EActorRepListTypeFlags::Default);
+		const TArrayView<const FActorRepListType> DependentActors = ListGatherer.ViewActors(EActorRepListTypeFlags::Default);
 
-		for (const FActorRepListConstView& DependentActorList : DependentLists)
+		for (AActor* DependentActor : DependentActors)
 		{
-			for (AActor* DependentActor : DependentActorList)
-			{
-				repCheck(DependentActor);
+			repCheck(DependentActor);
 
-				FConnectionReplicationActorInfo& DependentActorConnectionInfo = ConnectionActorInfoMap.FindOrAdd(DependentActor);
-				FGlobalActorReplicationInfo& DependentActorGlobalData = GlobalActorReplicationInfoMap.Get(DependentActor);
+			FConnectionReplicationActorInfo& DependentActorConnectionInfo = ConnectionActorInfoMap.FindOrAdd(DependentActor);
+			FGlobalActorReplicationInfo& DependentActorGlobalData = GlobalActorReplicationInfoMap.Get(DependentActor);
 
-				UpdateActorChannelCloseFrameNum(DependentActor, DependentActorConnectionInfo, DependentActorGlobalData, FrameNum, NetConnection);
+			UpdateActorChannelCloseFrameNum(DependentActor, DependentActorConnectionInfo, DependentActorGlobalData, FrameNum, NetConnection);
 
-				// Dependent actor channel will stay open as long as the owning actor channel is open
-				DependentActorConnectionInfo.ActorChannelCloseFrameNum = FMath::Max<uint32>(CloseFrameNum, DependentActorConnectionInfo.ActorChannelCloseFrameNum);
-
-				if (!ReadyForNextReplication(DependentActorConnectionInfo, DependentActorGlobalData, FrameNum))
-				{
-					continue;
-				}
+			// Dependent actor channel will stay open as long as the owning actor channel is open
+			DependentActorConnectionInfo.ActorChannelCloseFrameNum = FMath::Max<uint32>(CloseFrameNum, DependentActorConnectionInfo.ActorChannelCloseFrameNum);
 
-				if (!ensureMsgf(IsActorValidForReplication(DependentActor), TEXT("DependentActor %s (Owner: %s) not valid for replication (BeingDestroyed:%d) (IsValid:%d) (Unreachable:%d) (TearOff:%d)! Channel = %s"),
-					*DependentActor->GetFullName(), *Actor->GetFullName(),
-					DependentActor->IsActorBeingDestroyed(), IsValid(DependentActor), DependentActor->IsUnreachable(), DependentActor->GetTearOff(),
-					*DescribeSafe(DependentActorConnectionInfo.Channel)))
-				{
-					continue;
-				}
+			if (!ReadyForNextReplication(DependentActorConnectionInfo, DependentActorGlobalData, FrameNum))
+			{
+				continue;
+			}
 
-				//UE_LOG(LogReplicationGraph, Display, TEXT("DependentActor %s %s. NextReplicationFrameNum: %d. FrameNum: %d. ForceNetUpdateFrame: %d. LastRepFrameNum: %d."), *DependentActor->GetPathName(), *NetConnection->GetName(), DependentActorConnectionInfo.NextReplicationFrameNum, FrameNum, DependentActorGlobalData.ForceNetUpdateFrame, DependentActorConnectionInfo.LastRepFrameNum);
-				BitsWritten += ReplicateSingleActor(DependentActor, DependentActorConnectionInfo, DependentActorGlobalData, ConnectionActorInfoMap, ConnectionManager, FrameNum);
+			if (!ensureMsgf(IsActorValidForReplication(DependentActor), TEXT("DependentActor %s (Owner: %s) not valid for replication (BeingDestroyed:%d) (IsValid:%d) (Unreachable:%d) (TearOff:%d)! Channel = %s"),
+				*DependentActor->GetFullName(), *Actor->GetFullName(),
+				DependentActor->IsActorBeingDestroyed(), IsValid(DependentActor), DependentActor->IsUnreachable(), DependentActor->GetTearOff(),
+				*DescribeSafe(DependentActorConnectionInfo.Channel)))
+			{
+				continue;
 			}
+
+			//UE_LOG(LogReplicationGraph, Display, TEXT("DependentActor %s %s. NextReplicationFrameNum: %d. FrameNum: %d. ForceNetUpdateFrame: %d. LastRepFrameNum: %d."), *DependentActor->GetPathName(), *NetConnection->GetName(), DependentActorConnectionInfo.NextReplicationFrameNum, FrameNum, DependentActorGlobalData.ForceNetUpdateFrame, DependentActorConnectionInfo.LastRepFrameNum);
+			BitsWritten += ReplicateSingleActor(DependentActor, DependentActorConnectionInfo, DependentActorGlobalData, ConnectionActorInfoMap, ConnectionManager, FrameNum);
 		}
 	}
 
@@ -2197,14 +2185,11 @@ void UReplicationGraph::HandleStarvedActorList(const UNetReplicationGraphConnect
 
 		
 		GlobalActorInfo.GatherDependentActorLists(RepGraphConnection, DependentListContainer);
-		const TArray<FActorRepListConstView>& DependentLists = DependentListContainer.GetLists(EActorRepListTypeFlags::Default);
-		for (const FActorRepListConstView& DependentList : DependentLists)
+		const TArrayView<const FActorRepListType> DependentActors = DependentListContainer.ViewActors(EActorRepListTypeFlags::Default);
+		for (AActor* DependentActor : DependentActors)
 		{
-			for (AActor* DependentActor : DependentList)
-			{
-				FConnectionReplicationActorInfo& DependentActorConnectionInfo = ConnectionActorInfoMap.FindOrAdd(DependentActor);
-				DependentActorConnectionInfo.ActorChannelCloseFrameNum = FMath::Max<uint32>(CloseFrameNum, DependentActorConnectionInfo.ActorChannelCloseFrameNum);
-			}
+			FConnectionReplicationActorInfo& DependentActorConnectionInfo = ConnectionActorInfoMap.FindOrAdd(DependentActor);
+			DependentActorConnectionInfo.ActorChannelCloseFrameNum = FMath::Max<uint32>(CloseFrameNum, DependentActorConnectionInfo.ActorChannelCloseFrameNum);
 		}
 		DependentListContainer.Reset();
 	}
@@ -4400,14 +4385,11 @@ void UReplicationGraphNode_DynamicSpatialFrequency::CalcFrequencyForActor(AActor
 
 				// Update the dependent actor's timeout frame too
 				GlobalInfo.GatherDependentActorLists(RepGraphConnection, DependentListContainer);
-				const TArray<FActorRepListConstView>& DependentLists = DependentListContainer.GetLists(EActorRepListTypeFlags::Default);
-				for (const FActorRepListConstView& DependentList : DependentLists)
+				const TArrayView<const FActorRepListType> DependentActors = DependentListContainer.ViewActors(EActorRepListTypeFlags::Default);
+				for (AActor* DependentActor : DependentActors)
 				{
-					for (AActor* DependentActor : DependentList)
-					{
-						FConnectionReplicationActorInfo& DependentActorConnectionInfo = ConnectionMap.FindOrAdd(DependentActor);
-						DependentActorConnectionInfo.ActorChannelCloseFrameNum = FMath::Max<uint32>(ConnectionInfo.ActorChannelCloseFrameNum, ConnectionInfo.NextReplicationFrameNum + 1);
-					}
+					FConnectionReplicationActorInfo& DependentActorConnectionInfo = ConnectionMap.FindOrAdd(DependentActor);
+					DependentActorConnectionInfo.ActorChannelCloseFrameNum = FMath::Max<uint32>(ConnectionInfo.ActorChannelCloseFrameNum, ConnectionInfo.NextReplicationFrameNum + 1);
 				}
 				DependentListContainer.Reset();
 
diff --git a/Engine/Plugins/Runtime/ReplicationGraph/Source/Private/ReplicationGraphTypes.cpp b/Engine/Plugins/Runtime/ReplicationGraph/Source/Private/ReplicationGraphTypes.cpp
--- a/Engine/Plugins/Runtime/ReplicationGraph/Source/Private/ReplicationGraphTypes.cpp
+++ b/Engine/Plugins/Runtime/ReplicationGraph/Source/Private/ReplicationGraphTypes.cpp
@@ -611,6 +611,30 @@ void FActorRepListStatCollector::VisitExplicitStreamingLevelList(FName ListOwner
 	StreamingLevelStats.MaxListSize = FMath::Max(StreamingLevelStats.MaxListSize, ListSize);
 }
 
+PRAGMA_DISABLE_DEPRECATION_WARNINGS
+const TArray<FActorRepListConstView>& FGatheredReplicationActorLists::GetLists(EActorRepListTypeFlags ListFlags) const
+{
+	// Synthesize an array of views of actors for backwards compatibility.
+	// All actors appear in a single FActorRepListConstView.
+	static FActorRepListRefView ActorListView;
+	ActorListView.Reset();
+
+	const TArrayView<const FActorRepListType> Actors = ViewActors(ListFlags);
+	ActorListView.Reserve(Actors.Num());
+	for (const FActorRepListType& Actor : Actors)
+	{
+		ActorListView.Add(Actor);
+	}
+
+	static TArray<FActorRepListConstView> ArrayOfViewsOfActors;
+	ArrayOfViewsOfActors.Reset();
+	ArrayOfViewsOfActors.Reserve(1);
+	ArrayOfViewsOfActors.Add(FActorRepListConstView(ActorListView));
+
+	return ArrayOfViewsOfActors;
+}
+PRAGMA_ENABLE_DEPRECATION_WARNINGS
+
 void FLevelBasedActorList::AddNetworkActor(AActor* NetActor)
 {
 	FNewReplicatedActorInfo ActorInfo(NetActor);
diff --git a/Engine/Plugins/Runtime/ReplicationGraph/Source/Public/ReplicationGraphTypes.h b/Engine/Plugins/Runtime/ReplicationGraph/Source/Public/ReplicationGraphTypes.h
--- a/Engine/Plugins/Runtime/ReplicationGraph/Source/Public/ReplicationGraphTypes.h
+++ b/Engine/Plugins/Runtime/ReplicationGraph/Source/Public/ReplicationGraphTypes.h
@@ -485,8 +485,9 @@ private:
  * Gives temporary read-only access to a FActorRepListRefView by holding a reference to it.
  */
 struct REPLICATIONGRAPH_API FActorRepListConstView
+UE_DEPRECATED(5.4, "Use TArrayView<const FActorRepListType> instead")
 {
-	FActorRepListConstView(const FActorRepListRefView& InListReferenced) :
+	explicit FActorRepListConstView(const FActorRepListRefView& InListReferenced) :
 		ListReferenced(InListReferenced)
 	{}
 
@@ -507,8 +508,6 @@ private:
 	const FActorRepListRefView& ListReferenced;
 };
 
-/** A read only, non owning (ref counting) view to an actor replication list: essentially a raw pointer and the category of the list. These are only created *from* FActorRepListRefView */
-
 #if !(UE_BUILD_SHIPPING || UE_BUILD_TEST)
 	// Intended to be called from immediate mode window while debugging
 	extern "C" DLLEXPORT void PrintRepListDetails(int32 PoolSize, int32 BlockIdx, int32 ListIdx);
@@ -524,7 +523,7 @@ private:
 // --------------------------------------------------------------------------------------------------------------------------------------------
 // --------------------------------------------------------------------------------------------------------------------------------------------
 
-// This represents "the list of gathered lists". This is what we push down the Replication Graph and nodes will either Push/Pop List Categories or will add their Replication Lists.
+// This represents "the list of gathered actors". This is what we push down the Replication Graph and nodes will either Push/Pop List Categories or will add their Replication Lists.
 struct REPLICATIONGRAPH_API FGatheredReplicationActorLists
 {
 	void AddReplicationActorList(const FActorRepListRefView& List, EActorRepListTypeFlags Flags = EActorRepListTypeFlags::Default)
@@ -533,11 +532,7 @@ struct REPLICATIONGRAPH_API FGatheredReplicationActorLists
 		if (CVar_RepGraph_Verify)
 			List.VerifyContents_Slow();
 #endif
-		if (List.Num() > 0)
-		{
-			ReplicationLists[(uint32)Flags].Emplace(FActorRepListConstView(List));
-			CachedNum++;
-		}
+		List.AppendToTArray(ReplicationLists[(uint32)Flags]);
 	}
 
 	FORCEINLINE void Reset()
@@ -546,17 +541,22 @@ struct REPLICATIONGRAPH_API FGatheredReplicationActorLists
 		{
 			ReplicationLists[i].Reset();
 		}
-		CachedNum = 0;
 	}
 	FORCEINLINE int32 NumLists() const
 	{
-		return CachedNum;
+		return ReplicationLists.Num();
 	}
 
-	FORCEINLINE const TArray<FActorRepListConstView>& GetLists(EActorRepListTypeFlags ListFlags) const
+PRAGMA_DISABLE_DEPRECATION_WARNINGS
+	UE_DEPRECATED(5.4, "Use ViewActors() instead")
+	FORCEINLINE const TArray<FActorRepListConstView>& GetLists(EActorRepListTypeFlags ListFlags) const;
+PRAGMA_ENABLE_DEPRECATION_WARNINGS
+
+	FORCEINLINE TArrayView<const FActorRepListType> ViewActors(EActorRepListTypeFlags ListFlags) const
 	{
 		return ReplicationLists[(uint32)ListFlags];
 	}
+
 	FORCEINLINE bool ContainsLists(EActorRepListTypeFlags Flags) const
 	{
 		return ReplicationLists[(uint32)Flags].Num() > 0;
@@ -564,8 +564,7 @@ struct REPLICATIONGRAPH_API FGatheredReplicationActorLists
 
 private:
 
-	TStaticArray< TArray<FActorRepListConstView>, (uint32)EActorRepListTypeFlags::Max > ReplicationLists;
-	int32 CachedNum = 0;
+	TStaticArray< TArray<FActorRepListType>, (uint32)EActorRepListTypeFlags::Max > ReplicationLists;
 };
 
 // --------------------------------------------------------------------------------------------------------------------------------------------
@@ -1221,20 +1220,17 @@ struct FGlobalActorReplicationInfoMap
 		FGatheredReplicationActorLists ListContainer;
 		MainActorInfo->DependentActorList.AppendAllLists(ListContainer);
 
-		// Remove the actor from his child dependents
-		const TArray<FActorRepListConstView>& DependentActorLists = ListContainer.GetLists(EActorRepListTypeFlags::Default);
-		for (const FActorRepListConstView& DependentActorList : DependentActorLists)
+		// Remove the actor from its child dependents
+		const TArrayView<const FActorRepListType> DependentActorsList = ListContainer.ViewActors(EActorRepListTypeFlags::Default);
+		for (AActor* DependentActor : DependentActorsList)
 		{
-			for (AActor* DependentActor : DependentActorList)
+			if (FGlobalActorReplicationInfo* ChildInfo = Find(DependentActor))
 			{
-				if (FGlobalActorReplicationInfo* ChildInfo = Find(DependentActor))
-				{
-					ChildInfo->ParentActorList.RemoveSingleSwap(MainActor);
-				}
+				ChildInfo->ParentActorList.RemoveSingleSwap(MainActor);
 			}
 		}
 
-		// Remove the actor from his parents
+		// Remove the actor from its parents
 		FNewReplicatedActorInfo LevelActorInfo(MainActor);
 
 		for (AActor* ParentActor : MainActorInfo->ParentActorList)
