Remove deprecation attribute from type

diff --git a/Engine/Plugins/Runtime/ReplicationGraph/Source/Public/ReplicationGraphTypes.h b/Engine/Plugins/Runtime/ReplicationGraph/Source/Public/ReplicationGraphTypes.h
--- a/Engine/Plugins/Runtime/ReplicationGraph/Source/Public/ReplicationGraphTypes.h
+++ b/Engine/Plugins/Runtime/ReplicationGraph/Source/Public/ReplicationGraphTypes.h
@@ -483,8 +483,8 @@ private:
 
 /**
  * Gives temporary read-only access to a FActorRepListRefView by holding a reference to it.
+ * Deprecated: Use TArrayView<const FActorRepListType> instead.
  */
-UE_DEPRECATED(5.4, "Use TArrayView<const FActorRepListType> instead")
 struct REPLICATIONGRAPH_API FActorRepListConstView
 {
 	explicit FActorRepListConstView(const FActorRepListRefView& InListReferenced) :
