Replication graph properly handles renaming statically-spatialized actors.
Also, put more of the invalid actor logging behind the Net.RepGraph.LogDebugInfoPeriod to prevent log spam every frame an invalid actor is detected in the graph.

diff --git a/Engine/Plugins/Runtime/ReplicationGraph/Source/Private/ReplicationGraph.cpp b/Engine/Plugins/Runtime/ReplicationGraph/Source/Private/ReplicationGraph.cpp
--- a/Engine/Plugins/Runtime/ReplicationGraph/Source/Private/ReplicationGraph.cpp
+++ b/Engine/Plugins/Runtime/ReplicationGraph/Source/Private/ReplicationGraph.cpp
@@ -249,31 +249,6 @@ void UpdateActorConnectionCounter(AActor* InActor, UNetConnection* InConnection,
 void LogMoreInfoOnIsActorValidFailure(const FActorRepListType& In)
 {
 #if WITH_SERVER_CODE
-	// Always log information about the actor (it's not much)
-	if (!DoesActorPointerLookValid(In))
-	{
-#if UE_ACTOR_REPLIST_TYPE_EXTRA_SAFETY
-		UE_LOG(LogReplicationGraph, Error, TEXT("Invalid actor pointer detected during replication: Ptr = %p, Name='%s', Owner='%s', OuterPackage='%s'"),
-			In.ActorRaw, *In.ActorName.ToString(), *In.OwnerName.ToString(), *In.OuterPackageName.ToString()
-			);
-#else
-		UE_LOG(LogReplicationGraph, Error, TEXT("Invalid actor detected during replication: Ptr = %p"), static_cast<AActor*>(In));
-#endif
-	}
-	else
-	{
-		// Actor pointer is valid, but some of its properties may be not 
-		AActor* Actor = In;
-		UE_LOG(LogReplicationGraph, Error, TEXT("Actor not valid for replication (BeingDestroyed:%d) (IsValid:%d) (Unreachable:%d) (TearOff:%d)! Actor = %s"),
-			Actor->IsActorBeingDestroyed(), IsValid(Actor), Actor->IsUnreachable(), Actor->GetTearOff(),
-			*Actor->GetFullName());
-#if UE_ACTOR_REPLIST_TYPE_EXTRA_SAFETY
-		UE_LOG(LogReplicationGraph, Error, TEXT("More info about invalid actor '%s': Owner='%s', OuterPackage='%s'"),
-			*In.ActorName.ToString(), *In.OwnerName.ToString(), *In.OuterPackageName.ToString());
-#endif
-	}
-
-
 	static std::atomic<double> LastTimeLogged = 0;
 	double CurrentTime = FPlatformTime::Seconds();
 	if (CVar_RepGraph_LogDebugInfoPeriod > 0 && ((CurrentTime - LastTimeLogged) > double(CVar_RepGraph_LogDebugInfoPeriod)))
@@ -282,13 +257,27 @@ void LogMoreInfoOnIsActorValidFailure(const FActorRepListType& In)
 
 		if (DoesActorPointerLookValid(In))
 		{
+			// Actor pointer is valid, but some of its properties may be not 
 			AActor* Actor = In;
-			UE_LOG(LogReplicationGraph, Log, TEXT("Invalid actor found at the last moment, executing Net.RepGraph.PrintAll:"));
+			UE_LOG(LogReplicationGraph, Error, TEXT("Actor not valid for replication (BeingDestroyed:%d) (IsValid:%d) (Unreachable:%d) (TearOff:%d)! Actor = %s"),
+				Actor->IsActorBeingDestroyed(), IsValid(Actor), Actor->IsUnreachable(), Actor->GetTearOff(),
+				*Actor->GetFullName());
+#if UE_ACTOR_REPLIST_TYPE_EXTRA_SAFETY
+			UE_LOG(LogReplicationGraph, Error, TEXT("More info about invalid actor '%s': Owner='%s', OuterPackage='%s'"),
+				*In.ActorName.ToString(), *In.OwnerName.ToString(), *In.OuterPackageName.ToString());
+#endif
+
 			GEngine->Exec(Actor->GetWorld(), TEXT("Net.RepGraph.PrintAll"));
 		}
 		else
 		{
-			UE_LOG(LogReplicationGraph, Warning, TEXT("Invalid actor found at the last moment, not executing Net.RepGraph.PrintAll due to invalid actor pointer."));
+#if UE_ACTOR_REPLIST_TYPE_EXTRA_SAFETY
+			UE_LOG(LogReplicationGraph, Error, TEXT("Invalid actor pointer detected during replication: Ptr = %p, Name='%s', Owner='%s', OuterPackage='%s'"),
+				In.ActorRaw, *In.ActorName.ToString(), *In.OwnerName.ToString(), *In.OuterPackageName.ToString()
+				);
+#else
+			UE_LOG(LogReplicationGraph, Error, TEXT("Invalid actor pointer detected during replication: Ptr = %p"), static_cast<AActor*>(In));
+#endif
 		}
 
 		UE_LOG(LogReplicationGraph, Log, TEXT("Invalid actor found at the last moment, printing (game-specific) current routing:"));
@@ -5120,8 +5109,12 @@ void UReplicationGraphNode_GridCell::RenameStaticActor(const FRenamedReplicatedA
 		GetDormancyNode()->RenameDormantActor(ActorInfo);
 	}
 	else
-	{	
-		Super::NotifyActorRenamed(ActorInfo);
+	{
+		const bool bRemovedSomething = Super::NotifyRemoveNetworkActor(ActorInfo.OldActorInfo, true);
+		if (ensureMsgf(bRemovedSomething, TEXT("Renamed static actor %s not found in non-dormant grid cell list."), *GetFullNameSafe(ActorInfo.OldActorInfo.GetActor())))
+		{
+			Super::NotifyAddNetworkActor(ActorInfo.NewActorInfo);
+		}
 	}
 }
 
