Fix for invalid actors remaining in the replication graph if:
-Net.RepGraph.HandleDynamicActorRename is enabled, and
-A replicated actor is added before it's initialized and is routed as a dormant, spatialized actor (which is put on the ReplicationGraphNode_GridSpatialization2D::PendingStaticSpatializedActors list), and
-That actor is then destroyed.

Added automation test "Net.RenameReplicatedActor.EarlyRenameAndDestroy" to EngineTest to test this case and some supporting changes:
-Enable ReplicationGraph plugin for the EngineTest project
-Reset the timer for replication graph logging invalid actor info when a graph is initialized. This ensures the logs appear if multiple tests are run in quick succession.

diff --git a/Engine/Plugins/Runtime/ReplicationGraph/Source/Private/ReplicationGraph.cpp b/Engine/Plugins/Runtime/ReplicationGraph/Source/Private/ReplicationGraph.cpp
--- a/Engine/Plugins/Runtime/ReplicationGraph/Source/Private/ReplicationGraph.cpp
+++ b/Engine/Plugins/Runtime/ReplicationGraph/Source/Private/ReplicationGraph.cpp
@@ -74,6 +74,9 @@ namespace UE::Net::Private
 	// TInlineAllocator's size is optimized for default DestroyDormantDynamicActorsCellTTL value
 	constexpr int32 RepGraphDormancyNodesInlineBufferSize = 200;
 
+	// Track time for CVar_RepGraph_LogDebugInfoPeriod. Even though it's static, will be reset when a new repgraph instance is created
+	static std::atomic<double> LastTimeLoggedMoreInfo = 0;
+
 	int32 CVar_RepGraph_HandleDynamicActorRename = 0;
 	static FAutoConsoleVariableRef CVarRepGraphHandleDynamicActorRename(TEXT("Net.RepGraph.HandleDynamicActorRename"), CVar_RepGraph_HandleDynamicActorRename, TEXT("If nonzero, when a dynamic actor's outer/level changes, repgraph will update its cached level information."));
 }
@@ -249,11 +252,12 @@ void UpdateActorConnectionCounter(AActor* InActor, UNetConnection* InConnection,
 void LogMoreInfoOnIsActorValidFailure(const FActorRepListType& In)
 {
 #if WITH_SERVER_CODE
-	static std::atomic<double> LastTimeLogged = 0;
+	using namespace UE::Net::Private;
+
 	double CurrentTime = FPlatformTime::Seconds();
-	if (CVar_RepGraph_LogDebugInfoPeriod > 0 && ((CurrentTime - LastTimeLogged) > double(CVar_RepGraph_LogDebugInfoPeriod)))
+	if (CVar_RepGraph_LogDebugInfoPeriod > 0 && ((CurrentTime - LastTimeLoggedMoreInfo) > double(CVar_RepGraph_LogDebugInfoPeriod)))
 	{
-		LastTimeLogged = CurrentTime;
+		LastTimeLoggedMoreInfo = CurrentTime;
 
 		if (DoesActorPointerLookValid(In))
 		{
@@ -380,6 +384,8 @@ void UReplicationGraph::InitForNetDriver(UNetDriver* InNetDriver)
 	InitGlobalActorClassSettings();
 	InitGlobalGraphNodes();
 
+	UE::Net::Private::LastTimeLoggedMoreInfo = 0;
+
 	for (UNetConnection* ClientConnection : NetDriver->ClientConnections)
 	{
 		AddClientConnection(ClientConnection);
@@ -5463,6 +5469,16 @@ void UReplicationGraphNode_GridSpatialization2D::RenameActor_Static(const FRenam
 	}
 	else
 	{
+		// May have been a pending actor
+		for (const FPendingStaticActors& Pending : PendingStaticSpatializedActors)
+		{
+			if (Pending.Actor == ActorInfo.OldActorInfo.Actor)
+			{
+				// Pending static actors don't store any level/outer information, so there's nothing to do
+				return;
+			}
+		}
+
 		UE_LOG(LogReplicationGraph, Warning, TEXT("UReplicationGraphNode_GridSpatialization2D::RenameActor_Static attempted rename %s from static list but it was not there."), *GetActorRepListTypeDebugString(ActorInfo.NewActorInfo.Actor));
 		FCachedDynamicActorInfo* DynamicFoundInfo = DynamicSpatializedActors.Find(ActorInfo.NewActorInfo.Actor);
 		if (DynamicFoundInfo)
