Deprecate type in proper location

diff --git a/Engine/Plugins/Runtime/ReplicationGraph/Source/Public/ReplicationGraphTypes.h b/Engine/Plugins/Runtime/ReplicationGraph/Source/Public/ReplicationGraphTypes.h
--- a/Engine/Plugins/Runtime/ReplicationGraph/Source/Public/ReplicationGraphTypes.h
+++ b/Engine/Plugins/Runtime/ReplicationGraph/Source/Public/ReplicationGraphTypes.h
@@ -487,9 +487,11 @@ private:
 
 /**
  * Gives temporary read-only access to a FActorRepListRefView by holding a reference to it.
- * Deprecated: Use TArrayView<const FActorRepListType> instead.
  */
-struct REPLICATIONGRAPH_API FActorRepListConstView
+struct
+REPLICATIONGRAPH_API 
+UE_DEPRECATED(5.4, "Use TArrayView<const FActorRepListType> instead")
+FActorRepListConstView
 {
 	explicit FActorRepListConstView(const FActorRepListRefView& InListReferenced) :
 		ListReferenced(InListReferenced)
