Add an option to use weak pointers in RepGraph to catch dangling pointers without crashing.

Adds resilience to RepGraph if there are actor lifetime bugs.

Off by default. Enabled by cvar `Net.RepGraph.UseWeakPointers true`. Can be changed at runtime.

diff --git a/Engine/Plugins/Runtime/ReplicationGraph/Source/Private/ReplicationGraph.cpp b/Engine/Plugins/Runtime/ReplicationGraph/Source/Private/ReplicationGraph.cpp
--- a/Engine/Plugins/Runtime/ReplicationGraph/Source/Private/ReplicationGraph.cpp
+++ b/Engine/Plugins/Runtime/ReplicationGraph/Source/Private/ReplicationGraph.cpp
@@ -277,7 +277,7 @@ void LogMoreInfoOnIsActorValidFailure(const FActorRepListType& In)
 		{
 #if UE_ACTOR_REPLIST_TYPE_EXTRA_SAFETY
 			UE_LOG(LogReplicationGraph, Error, TEXT("Invalid actor pointer detected during replication: Ptr = %p, Name='%s', Owner='%s', OuterPackage='%s'"),
-				In.ActorRaw, *In.ActorName.ToString(), *In.OwnerName.ToString(), *In.OuterPackageName.ToString()
+				In.GetActor(), *In.ActorName.ToString(), *In.OwnerName.ToString(), *In.OuterPackageName.ToString()
 				);
 #else
 			UE_LOG(LogReplicationGraph, Error, TEXT("Invalid actor pointer detected during replication: Ptr = %p"), static_cast<AActor*>(In));
@@ -773,6 +773,15 @@ void UReplicationGraph::RemoveNetworkActor(AActor* Actor)
 {
 	QUICK_SCOPE_CYCLE_COUNTER(UReplicationGraph_RemoveNetworkActor);
 
+#if UE_BUILD_DEBUG || UE_BUILD_DEVELOPMENT
+	UE_CLOG(
+		Actor->IsPendingKillPending(),
+		LogReplicationGraph,
+		Verbose,
+		TEXT("Removing Actor that is already pending kill or invalid, please remove it before destroying it. Actor=%s"),
+		*Actor->GetActorNameOrLabel());
+#endif
+
 	if (ActiveNetworkActors.Remove(Actor) == 0)
 	{
 		// Guarding against double removes
@@ -5276,8 +5285,13 @@ void UReplicationGraphNode_GridSpatialization2D::RemoveActor_Static(const FNewRe
 
 	if (GraphGlobals.IsValid())
 	{
-		FGlobalActorReplicationInfo& GlobalInfo = GraphGlobals->GlobalActorReplicationInfoMap->Get(ActorInfo.Actor);
-		RemoveActorInternal_Static(ActorInfo, GlobalInfo, GlobalInfo.bWantsToBeDormant); 
+		FGlobalActorReplicationInfo* GlobalInfo = GraphGlobals->GlobalActorReplicationInfoMap->Find(ActorInfo.Actor);
+		if (!GlobalInfo)
+		{
+			return;
+		}
+
+		RemoveActorInternal_Static(ActorInfo, *GlobalInfo, GlobalInfo->bWantsToBeDormant); 
 	}
 }
 
@@ -5287,10 +5301,15 @@ void UReplicationGraphNode_GridSpatialization2D::RemoveActor_Dormancy(const FNew
 
 	if (GraphGlobals.IsValid())
 	{
-		FGlobalActorReplicationInfo& ActorRepInfo = GraphGlobals->GlobalActorReplicationInfoMap->Get(ActorInfo.Actor);
-		if (ActorRepInfo.bWantsToBeDormant)
+		FGlobalActorReplicationInfo* ActorRepInfo = GraphGlobals->GlobalActorReplicationInfoMap->Find(ActorInfo.Actor);
+		if (!ActorRepInfo)
+		{
+			return;
+		}
+
+		if (ActorRepInfo->bWantsToBeDormant)
 		{
-			RemoveActorInternal_Static(ActorInfo, ActorRepInfo, true);
+			RemoveActorInternal_Static(ActorInfo, *ActorRepInfo, true);
 		}
 		else
 		{
@@ -5301,7 +5320,7 @@ void UReplicationGraphNode_GridSpatialization2D::RemoveActor_Dormancy(const FNew
 		// This means that even if AddActor_Dormancy is called multiple times with the same Actor, a single call to RemoveActor_Dormancy
 		// will completely remove the Actor from either the Static or Dynamic list appropriately.
 		// Therefore, it should be safe to call RemoveAll and not worry about trying to track individual delegate handles.
-		ActorRepInfo.Events.DormancyChange.RemoveAll(this);
+		ActorRepInfo->Events.DormancyChange.RemoveAll(this);
 	}
 }
 
diff --git a/Engine/Plugins/Runtime/ReplicationGraph/Source/Private/ReplicationGraphTypes.cpp b/Engine/Plugins/Runtime/ReplicationGraph/Source/Private/ReplicationGraphTypes.cpp
--- a/Engine/Plugins/Runtime/ReplicationGraph/Source/Private/ReplicationGraphTypes.cpp
+++ b/Engine/Plugins/Runtime/ReplicationGraph/Source/Private/ReplicationGraphTypes.cpp
@@ -17,6 +17,15 @@ DEFINE_LOG_CATEGORY( LogReplicationGraph );
 DECLARE_LLM_MEMORY_STAT(TEXT("NetRepGraph"), STAT_NetRepGraphLLM, STATGROUP_LLMFULL);
 LLM_DEFINE_TAG(NetRepGraph, NAME_None, TEXT("Networking"), GET_STATFNAME(STAT_NetRepGraphLLM), GET_STATFNAME(STAT_NetworkingSummaryLLM));
 
+#if WITH_SERVER_CODE && UE_ACTOR_REPLIST_TYPE_EXTRA_SAFETY
+
+static FAutoConsoleVariableRef CVar_NetRepGraphUseWeakPointers(
+	TEXT("Net.RepGraph.UseWeakPointers"),
+	UE::Net::RepGraph::bUseWeakPointers,
+	TEXT("Uses weak pointers in RepGraph, validating them on access to detect dangling pointers without crashing."));
+
+#endif
+
 // --------------------------------------------------------------------------------------------------------------------------------------------
 // --------------------------------------------------------------------------------------------------------------------------------------------
 // Actor List Allocator
diff --git a/Engine/Plugins/Runtime/ReplicationGraph/Source/Public/ReplicationGraphTypes.h b/Engine/Plugins/Runtime/ReplicationGraph/Source/Public/ReplicationGraphTypes.h
--- a/Engine/Plugins/Runtime/ReplicationGraph/Source/Public/ReplicationGraphTypes.h
+++ b/Engine/Plugins/Runtime/ReplicationGraph/Source/Public/ReplicationGraphTypes.h
@@ -9,8 +9,10 @@
 #include "Templates/Greater.h"
 #include "Templates/Requires.h"
 #include "UObject/Package.h"
+#include "UObject/WeakObjectPtr.h"
 
 #include <type_traits>
+#include <variant> // Switching to TVariant once it's trivially destructible
 
 #include "ReplicationGraphTypes.generated.h"
 
@@ -85,14 +87,30 @@ FORCEINLINE bool DoesActorPointerLookValid(const AActor* In)
 {
 	return ((uint64)(In) & 0x0F) == 0;
 }
-#else
+#else // UE_ACTOR_REPLIST_TYPE_EXTRA_SAFETY
+
+namespace UE::Net::RepGraph
+{
+	// Whether to store weak pointers in actor replication lists and validate them on use.
+	// When false, raw pointers are used instead.
+	inline bool bUseWeakPointers = false;
+}
+
 struct FActorRepListType
 {
-	/** Actual load-bearing payload */
-	AActor* ActorRaw;
-	/** Validity test */
-	FObjectKey ActorKey;
+private:
+	// Actor pointer (raw or weak).
+	//
+	// Most accesses should be via GetActor() or SetActor().
+	//
+	// Type is set during construction depending on bUseWeakPointers: FWeakObjectPtr when true, AActor* when false. Most functions should depend on
+	// the actual type, only SetActor() depends on bUseWeakPointers so the configuration can change at runtime without breaking extant instances.
+	//
+	// Requires std::variant rather than TVariant to be trivially destructible for FActorRepList's variable-length FActorRepListType[] Data member.
+	// Could be condensed to sizeof(AActor*) aka sizeof(FWeakObjectPtr) by storing the type flag in the high bit and masking it out.
+	std::variant<AActor*, FWeakObjectPtr> ActorUnion;
 
+public:
 	/** More info for debugging - ActorRaw's name */
 	FName ActorName;
 	/** More info for debugging - ActorRaw's Owner name */
@@ -100,32 +118,15 @@ struct FActorRepListType
 	/** More info for debugging - ActorRaw's Outer's package name */
 	FName OuterPackageName;
 
-	inline void SetDebugInfo()
+	FActorRepListType()
 	{
-		ActorName = OwnerName = OuterPackageName = NAME_None;
-		if (LIKELY(ActorRaw))
-		{
-			ActorName = ActorRaw->GetFName();
-			OwnerName = ActorRaw->GetOwner() ? ActorRaw->GetOwner()->GetFName() : NAME_None;
-			if (LIKELY(ActorRaw->GetOuter()))
-			{
-				// judging by the implementation GetPackage() cannot return nullptr, but play it safe
-				UPackage* Pkg = ActorRaw->GetOuter()->GetPackage();
-				if (LIKELY(Pkg))
-				{
-					OuterPackageName = Pkg->GetFName();
-				}
-			}
-		}
+		SetActor(nullptr);
 	}
 
-	FActorRepListType() = default;
-
-	FActorRepListType(AActor* InActor)
-		: ActorRaw(InActor)
-		, ActorKey(InActor)
+	// Permits implicit construction so FActorRepListType can be a drop-in substitute for `AActor*` depending on an ifdef.
+	/*implicit*/ FActorRepListType(AActor* InActor)
 	{
-		SetDebugInfo();
+		SetActor(InActor);
 	}
 
 	// to support conversion from TObjectPtr<ASubclassOfActor>
@@ -134,56 +135,143 @@ struct FActorRepListType
 		UE_REQUIRES(std::is_convertible_v<const T&, AActor*>)
 	>
 	FActorRepListType(const T& InActor)
-		: ActorRaw(InActor)
-		, ActorKey(InActor)
 	{
-		SetDebugInfo();
+		SetActor(InActor);
 	}
 
-	operator AActor*() { return ActorRaw; }
-	operator AActor*() const { return ActorRaw; }
-	AActor* operator->() { return ActorRaw; }
-	AActor* operator->() const { return ActorRaw; }
-	operator uint64() const { return reinterpret_cast<uint64>(ActorRaw); }
+	operator AActor* () const { return GetActor(); }
+	AActor* operator->() const { return GetActor(); }
+	explicit operator uint64() const { return reinterpret_cast<uint64>(GetActor()); }
 	FActorRepListType& operator=(FActorRepListType const& InActor) = default;
 	FActorRepListType& operator=(AActor* InActor)
 	{
-		ActorRaw = InActor;
-		ActorKey = InActor;
-		SetDebugInfo();
+		SetActor(InActor);
 		return *this;
 	}
 	FActorRepListType& operator=(TObjectPtr<AActor> InActor)
 	{
-		ActorRaw = InActor;
-		ActorKey = InActor;
-		SetDebugInfo();
+		SetActor(InActor.Get());
 		return *this;
 	}
-	bool operator==(FActorRepListType const& Other) const
+
+	friend bool operator==(const FActorRepListType& Left, const FActorRepListType& Right)
+	{
+		// Equality comparison uses the resolved pointer to ensure comparison is accurate if one is weak and one is raw; comparison only on the
+		// Variant object would indicate inequality if the types differed.
+		return Left.GetActor() == Right.GetActor();
+	}
+	friend bool operator!=(const FActorRepListType& Left, const FActorRepListType& Right)
+	{
+		return !(Left == Right);
+	}
+
+	friend bool operator==(const FActorRepListType& RepListActor, AActor* RawActor)
+	{
+		return RepListActor.GetActor() == RawActor;
+	}
+	friend bool operator!=(const FActorRepListType& RepListActor, AActor* RawActor)
+	{
+		return !(RepListActor == RawActor);
+	}
+
+	friend bool operator==(AActor* RawActor, const FActorRepListType& RepListActor)
 	{
-		return ActorRaw == Other.ActorRaw;
+		return RawActor == RepListActor.GetActor();
 	}
-	bool operator==(AActor* Other) const
+	friend bool operator!=(AActor* RawActor, const FActorRepListType& RepListActor)
 	{
-		return ActorRaw == Other;
+		return !(RepListActor == RawActor);
 	}
+
 	bool IsValid() const
 	{
-		UObject const* Object = ActorKey.ResolveObjectPtr();
-		return Object && Object == static_cast<UObject const*>(ActorRaw);
+		if (std::holds_alternative<AActor*>(ActorUnion))
+		{
+			return std::get<AActor*>(ActorUnion) != nullptr;
+		}
+
+		return std::get<FWeakObjectPtr>(ActorUnion).IsValid();
+	}
+
+	inline AActor* GetActor() const
+	{
+		if (std::holds_alternative<AActor*>(ActorUnion))
+		{
+			return std::get<AActor*>(ActorUnion);
+		}
+
+		const FWeakObjectPtr& WeakActor = std::get<FWeakObjectPtr>(ActorUnion);
+
+		if (WeakActor.IsExplicitlyNull())
+		{
+			return nullptr;
+		}
+
+		// Permit pointers to objects marked for destruction but not destroyed yet.
+		constexpr bool bPermitGarbage = true;
+		AActor* Actor = Cast<AActor>(WeakActor.Get(bPermitGarbage));
+		ensureMsgf(
+			Actor,
+			TEXT("RepGraph contains a destroyed actor (already gone): Name=%s Owner=%s Package=%s"),
+			*ActorName.ToString(),
+			*OwnerName.ToString(),
+			*OuterPackageName.ToString());
+
+		return Actor;
+	}
+
+private:
+	void SetActor(AActor* Actor)
+	{
+		if (UE::Net::RepGraph::bUseWeakPointers)
+		{
+			ActorUnion.emplace<FWeakObjectPtr>(Actor);
+		}
+		else
+		{
+			ActorUnion.emplace<AActor*>(Actor);
+		}
+
+		SetDebugInfo();
+	}
+
+	void SetDebugInfo()
+	{
+		ActorName = OwnerName = OuterPackageName = NAME_None;
+		if (AActor* Actor = GetActor(); LIKELY(Actor))
+		{
+			ActorName = Actor->GetFName();
+			OwnerName = Actor->GetOwner() ? Actor->GetOwner()->GetFName() : NAME_None;
+			if (LIKELY(Actor->GetOuter()))
+			{
+				// judging by the implementation GetPackage() cannot return nullptr, but play it safe
+				UPackage* Pkg = Actor->GetOuter()->GetPackage();
+				if (LIKELY(Pkg))
+				{
+					OuterPackageName = Pkg->GetFName();
+				}
+			}
+		}
 	}
 };
-template< class T > FORCEINLINE T* Cast(const FActorRepListType& Src) { return Cast<T>(Src.ActorRaw); }
-template< class T > FORCEINLINE T* ExactCast(const FActorRepListType& Src) { return ExactCast<T>(Src.ActorRaw); }
-template< class T > FORCEINLINE T* CastChecked(const FActorRepListType& Src, ECastCheckedType::Type CheckType = ECastCheckedType::NullChecked) { return CastChecked<T>(Src.ActorRaw, CheckType); }
 
+template< class T > FORCEINLINE T* Cast(const FActorRepListType& Src) { return Cast<T>(Src.GetActor()); }
+template< class T > FORCEINLINE T* ExactCast(const FActorRepListType& Src) { return ExactCast<T>(Src.GetActor()); }
+template< class T > FORCEINLINE T* CastChecked(const FActorRepListType& Src, ECastCheckedType::Type CheckType = ECastCheckedType::NullChecked)
+{
+	return CastChecked<T>(Src.GetActor(), CheckType);
+}
 
 FORCEINLINE bool DoesActorPointerLookValid(const FActorRepListType& In)
 {
 	return In.IsValid();
 }
-#endif
+
+inline uint32 GetTypeHash(const FActorRepListType& Actor)
+{
+	return GetTypeHash(Actor.GetActor());
+}
+#endif // UE_ACTOR_REPLIST_TYPE_EXTRA_SAFETY
 
 FORCEINLINE FString GetActorRepListTypeDebugString(const FActorRepListType& In) { return GetNameSafe(In); }
 FORCEINLINE UClass* GetActorRepListTypeClass(const FActorRepListType& In) { return In->GetClass(); }
@@ -1474,8 +1562,11 @@ struct FReplicationGraphGlobalData
 /** Stores "full debug details" about how an actor was prioritized. This is not used in the actual replication code, just saved off for logging/debugging.  */
 struct FPrioritizedActorFullDebugDetails
 {
-	FPrioritizedActorFullDebugDetails(FActorRepListType InActor) : Actor(InActor) { }
-	bool operator==(const FActorRepListType& InActor) const { return Actor == InActor; }
+	explicit FPrioritizedActorFullDebugDetails(FActorRepListType InActor) : Actor(InActor) { }
+	friend bool operator==(const FPrioritizedActorFullDebugDetails& Details, const FActorRepListType& InActor)
+	{
+		return Details.Actor == InActor;
+	}
 
 	FActorRepListType Actor;
 	FVector::FReal DistanceSq = 0.f;
@@ -1504,7 +1595,7 @@ struct FPrioritizedActorFullDebugDetails
 /** Debug data about an actor that was skipped during the prioritization phase */
 struct FSkippedActorFullDebugDetails
 {
-	FSkippedActorFullDebugDetails(FActorRepListType InActor) : Actor(InActor) { }
+	explicit FSkippedActorFullDebugDetails(FActorRepListType InActor) : Actor(InActor) { }
 	FActorRepListType Actor;
 	bool bWasDormant = false; // If set, was skipped because it is dormant on this connection
 	float DistanceCulled = 0.f; // If set, was skipped due to distance culling
@@ -1514,8 +1605,8 @@ struct FSkippedActorFullDebugDetails
 /** Prioritized List of actors to replicate. This is what we actually use to replicate actors. */
 struct FPrioritizedRepList
 {
-	FPrioritizedRepList() { }
-	FPrioritizedRepList(const FPrioritizedRepList& Other) { Items = Other.Items; }
+	FPrioritizedRepList() = default;
+	explicit FPrioritizedRepList(const FPrioritizedRepList& Other) { Items = Other.Items; }
 
 	struct FItem
 	{
