Fix FActorRepListType expensive deep copies

Copy/move was round-tripping via the deep-lookup implicit AActor* conversion rather than trivial copying.

Also speeds up comparison by comparing weak pointers without resolving them.

diff --git a/Engine/Plugins/Runtime/ReplicationGraph/Source/Public/ReplicationGraphTypes.h b/Engine/Plugins/Runtime/ReplicationGraph/Source/Public/ReplicationGraphTypes.h
--- a/Engine/Plugins/Runtime/ReplicationGraph/Source/Public/ReplicationGraphTypes.h
+++ b/Engine/Plugins/Runtime/ReplicationGraph/Source/Public/ReplicationGraphTypes.h
@@ -129,6 +129,12 @@ public:
 		SetActor(InActor);
 	}
 
+	// Copies/moves are clones that don't need to validate the object.
+	FActorRepListType(const FActorRepListType& InActor) = default;
+	FActorRepListType(FActorRepListType&& InActor) = default;
+	FActorRepListType& operator=(const FActorRepListType& InActor) = default;
+	FActorRepListType& operator=(FActorRepListType&& InActor) = default;
+
 	// to support conversion from TObjectPtr<ASubclassOfActor>
 	template <
 		typename T
@@ -142,7 +148,6 @@ public:
 	operator AActor* () const { return GetActor(); }
 	AActor* operator->() const { return GetActor(); }
 	explicit operator uint64() const { return reinterpret_cast<uint64>(GetActor()); }
-	FActorRepListType& operator=(FActorRepListType const& InActor) = default;
 	FActorRepListType& operator=(AActor* InActor)
 	{
 		SetActor(InActor);
@@ -156,8 +161,13 @@ public:
 
 	friend bool operator==(const FActorRepListType& Left, const FActorRepListType& Right)
 	{
-		// Equality comparison uses the resolved pointer to ensure comparison is accurate if one is weak and one is raw; comparison only on the
-		// Variant object would indicate inequality if the types differed.
+		// Optimize for both pointers being raw or both being weak, without resolving any weak ptrs.
+		if (LIKELY(Left.ActorUnion.index() == Right.ActorUnion.index()))
+		{
+			return Left.ActorUnion == Right.ActorUnion;
+		}
+
+		// Differing types; weak ptr must be resolved.
 		return Left.GetActor() == Right.GetActor();
 	}
 	friend bool operator!=(const FActorRepListType& Left, const FActorRepListType& Right)
@@ -273,6 +283,13 @@ inline uint32 GetTypeHash(const FActorRepListType& Actor)
 }
 #endif // UE_ACTOR_REPLIST_TYPE_EXTRA_SAFETY
 
+// Ensure desired performance characteristics.
+static_assert(std::is_trivially_destructible_v<FActorRepListType>);
+static_assert(std::is_trivially_move_constructible_v<FActorRepListType>);
+static_assert(std::is_trivially_move_assignable_v<FActorRepListType>);
+static_assert(std::is_trivially_copy_constructible_v<FActorRepListType>);
+static_assert(std::is_trivially_copy_assignable_v<FActorRepListType>);
+
 FORCEINLINE FString GetActorRepListTypeDebugString(const FActorRepListType& In) { return GetNameSafe(In); }
 FORCEINLINE UClass* GetActorRepListTypeClass(const FActorRepListType& In) { return In->GetClass(); }
 
