Correct UE_DEPRECATED version from 5.4 to 5.5

#rnx

diff --git a/Engine/Plugins/Runtime/ReplicationGraph/Source/Public/ReplicationGraphTypes.h b/Engine/Plugins/Runtime/ReplicationGraph/Source/Public/ReplicationGraphTypes.h
--- a/Engine/Plugins/Runtime/ReplicationGraph/Source/Public/ReplicationGraphTypes.h
+++ b/Engine/Plugins/Runtime/ReplicationGraph/Source/Public/ReplicationGraphTypes.h
@@ -595,7 +595,7 @@ private:
  */
 struct
 REPLICATIONGRAPH_API 
-UE_DEPRECATED(5.4, "Use TArrayView<const FActorRepListType> instead")
+UE_DEPRECATED(5.5, "Use TArrayView<const FActorRepListType> instead")
 FActorRepListConstView
 {
 	explicit FActorRepListConstView(const FActorRepListRefView& InListReferenced) :
@@ -659,7 +659,7 @@ struct REPLICATIONGRAPH_API FGatheredReplicationActorLists
 	}
 
 PRAGMA_DISABLE_DEPRECATION_WARNINGS
-	UE_DEPRECATED(5.4, "Use ViewActors() instead")
+	UE_DEPRECATED(5.5, "Use ViewActors() instead")
 	FORCEINLINE const TArray<FActorRepListConstView>& GetLists(EActorRepListTypeFlags ListFlags) const;
 PRAGMA_ENABLE_DEPRECATION_WARNINGS
 
