Speed up Net.RepGraph.UseWeakPointers

Reduce the number of times weak pointers are resolved to objects within a single function; look up the raw pointer once and reuse it.

#rnx

diff --git a/Engine/Plugins/Runtime/ReplicationGraph/Source/Private/ReplicationGraph.cpp b/Engine/Plugins/Runtime/ReplicationGraph/Source/Private/ReplicationGraph.cpp
--- a/Engine/Plugins/Runtime/ReplicationGraph/Source/Private/ReplicationGraph.cpp
+++ b/Engine/Plugins/Runtime/ReplicationGraph/Source/Private/ReplicationGraph.cpp
@@ -259,10 +259,10 @@ void LogMoreInfoOnIsActorValidFailure(const FActorRepListType& In)
 	{
 		LastTimeLoggedMoreInfo = CurrentTime;
 
-		if (DoesActorPointerLookValid(In))
+		AActor* Actor = In;
+		if (Actor)
 		{
-			// Actor pointer is valid, but some of its properties may be not 
-			AActor* Actor = In;
+			// Actor pointer is valid, but some of its properties may not be.
 			UE_LOG(LogReplicationGraph, Error, TEXT("Actor not valid for replication (BeingDestroyed:%d) (IsValid:%d) (Unreachable:%d) (TearOff:%d)! Actor = %s"),
 				Actor->IsActorBeingDestroyed(), IsValid(Actor), Actor->IsUnreachable(), Actor->GetTearOff(),
 				*Actor->GetFullName());
@@ -276,11 +276,11 @@ void LogMoreInfoOnIsActorValidFailure(const FActorRepListType& In)
 		else
 		{
 #if UE_ACTOR_REPLIST_TYPE_EXTRA_SAFETY
-			UE_LOG(LogReplicationGraph, Error, TEXT("Invalid actor pointer detected during replication: Ptr = %p, Name='%s', Owner='%s', OuterPackage='%s'"),
-				In.GetActor(), *In.ActorName.ToString(), *In.OwnerName.ToString(), *In.OuterPackageName.ToString()
+			UE_LOG(LogReplicationGraph, Error, TEXT("Invalid actor pointer detected during replication: Name='%s', Owner='%s', OuterPackage='%s'"),
+				*In.ActorName.ToString(), *In.OwnerName.ToString(), *In.OuterPackageName.ToString()
 				);
 #else
-			UE_LOG(LogReplicationGraph, Error, TEXT("Invalid actor pointer detected during replication: Ptr = %p"), static_cast<AActor*>(In));
+			UE_LOG(LogReplicationGraph, Error, TEXT("Invalid actor pointer detected during replication"));
 #endif
 		}
 
diff --git a/Engine/Plugins/Runtime/ReplicationGraph/Source/Public/ReplicationGraphTypes.h b/Engine/Plugins/Runtime/ReplicationGraph/Source/Public/ReplicationGraphTypes.h
--- a/Engine/Plugins/Runtime/ReplicationGraph/Source/Public/ReplicationGraphTypes.h
+++ b/Engine/Plugins/Runtime/ReplicationGraph/Source/Public/ReplicationGraphTypes.h
@@ -83,10 +83,6 @@ LLM_DECLARE_TAG_API(NetRepGraph, REPLICATIONGRAPH_API);
 
 #if !UE_ACTOR_REPLIST_TYPE_EXTRA_SAFETY
 typedef AActor* FActorRepListType;
-FORCEINLINE bool DoesActorPointerLookValid(const AActor* In)
-{
-	return ((uint64)(In) & 0x0F) == 0;
-}
 #else // UE_ACTOR_REPLIST_TYPE_EXTRA_SAFETY
 
 namespace UE::Net::RepGraph
@@ -147,7 +143,7 @@ public:
 
 	operator AActor* () const { return GetActor(); }
 	AActor* operator->() const { return GetActor(); }
-	explicit operator uint64() const { return reinterpret_cast<uint64>(GetActor()); }
+	explicit operator uint64() const { return GetTypeHash(); }
 	FActorRepListType& operator=(AActor* InActor)
 	{
 		SetActor(InActor);
@@ -161,13 +157,14 @@ public:
 
 	friend bool operator==(const FActorRepListType& Left, const FActorRepListType& Right)
 	{
-		// Optimize for both pointers being raw or both being weak, without resolving any weak ptrs.
-		if (LIKELY(Left.ActorUnion.index() == Right.ActorUnion.index()))
+		// Avoid FWeakObjectPtr::operator==() because it resolves both weak pointers if they don't match.
+		if (const FWeakObjectPtr* LeftPtr = std::get_if<FWeakObjectPtr>(&Left.ActorUnion), *RightPtr = std::get_if<FWeakObjectPtr>(&Right.ActorUnion);
+			LeftPtr && RightPtr)
 		{
-			return Left.ActorUnion == Right.ActorUnion;
+			return LeftPtr->HasSameIndexAndSerialNumber(*RightPtr);
 		}
 
-		// Differing types; weak ptr must be resolved.
+		// Differing types or both raw pointers; weak pointer must be resolved.
 		return Left.GetActor() == Right.GetActor();
 	}
 	friend bool operator!=(const FActorRepListType& Left, const FActorRepListType& Right)
@@ -193,6 +190,30 @@ public:
 		return !(RepListActor == RawActor);
 	}
 
+	// comparison with nullptr without resolving weak pointer or implicitly constructing a new FActorRepListType(nullptr).
+	friend bool operator==(const FActorRepListType& RepListActor, nullptr_t)
+	{
+		if (const FWeakObjectPtr* WeakPtrPtr = std::get_if<FWeakObjectPtr>(&RepListActor.ActorUnion); WeakPtrPtr)
+		{
+			return WeakPtrPtr->IsExplicitlyNull();
+		}
+
+		AActor* const* RawActorPtr = std::get_if<AActor*>(&RepListActor.ActorUnion);
+		return ensure(RawActorPtr) && *RawActorPtr == nullptr;
+	}
+	friend bool operator!=(const FActorRepListType& RepListActor, nullptr_t)
+	{
+		return !(RepListActor == nullptr);
+	}
+	friend bool operator==(nullptr_t, const FActorRepListType& RepListActor)
+	{
+		return RepListActor == nullptr;
+	}
+	friend bool operator!=(nullptr_t, const FActorRepListType& RepListActor)
+	{
+		return !(RepListActor == nullptr);
+	}
+
 	bool IsValid() const
 	{
 		if (std::holds_alternative<AActor*>(ActorUnion))
@@ -230,6 +251,14 @@ public:
 		return Actor;
 	}
 
+	// Get hash. Prefer this be a member function than standalone to avoid implicit conversions to FActorRepListType.
+	uint32 GetTypeHash() const
+	{
+		// Must resolve the pointer to ensure hashing produces identical results regardless of whether the stored pointer is weak or raw.
+		// (Hashing FWeakObjectPtr::ObjectIndex for weak or UObjectBase::InternalIndex for raw would be OK, but both are inaccessible.)
+		return ::GetTypeHash(GetActor());
+	}
+
 private:
 	void SetActor(AActor* Actor)
 	{
@@ -271,16 +300,6 @@ template< class T > FORCEINLINE T* CastChecked(const FActorRepListType& Src, ECa
 {
 	return CastChecked<T>(Src.GetActor(), CheckType);
 }
-
-FORCEINLINE bool DoesActorPointerLookValid(const FActorRepListType& In)
-{
-	return In.IsValid();
-}
-
-inline uint32 GetTypeHash(const FActorRepListType& Actor)
-{
-	return GetTypeHash(Actor.GetActor());
-}
 #endif // UE_ACTOR_REPLIST_TYPE_EXTRA_SAFETY
 
 // Ensure desired performance characteristics.
@@ -302,10 +321,15 @@ enum class EActorRepListTypeFlags : uint8
 };
 
 // Tests if an actor is valid for replication: not pending kill, etc. Says nothing about wanting to replicate or should replicate, etc.
+FORCEINLINE bool IsActorValidForReplication(const AActor* Actor)
+{
+	return Actor && !Actor->IsActorBeingDestroyed() && IsValidChecked(Actor) && !Actor->IsUnreachable(); 
+}
 FORCEINLINE bool IsActorValidForReplication(const FActorRepListType& In)
-{ 
-	return DoesActorPointerLookValid(In) && !In->IsActorBeingDestroyed() && IsValidChecked(In) && !In->IsUnreachable(); 
+{
+	return IsActorValidForReplication(static_cast<const AActor*>(In));
 }
+
 REPLICATIONGRAPH_API void LogMoreInfoOnIsActorValidFailure(const FActorRepListType& In);
 FORCEINLINE bool IsActorValidForReplication_LogMoreInfo(const FActorRepListType& In)
 { 
@@ -320,29 +344,29 @@ FORCEINLINE bool IsActorValidForReplication_LogMoreInfo(const FActorRepListType&
 
 // Tests if an actor is valid for replication gathering. Meaning, it can be gathered from the replication graph and considered for replication.
 FORCEINLINE bool IsActorValidForReplicationGather(const FActorRepListType& In)
-{ 
-	if (In == nullptr)
-	{
+{
+	const AActor* Actor = In;
+
+	if (!Actor)
 		return false;
-	}
 
-	if (!IsActorValidForReplication(In))
+	if (!IsActorValidForReplication(Actor))
 		return false;
 
-	if (In->GetIsReplicated() == false)
+	if (Actor->GetIsReplicated() == false)
 		return false;
 
-	if (In->GetTearOff())
+	if (Actor->GetTearOff())
 		return false;
 
-	if (In->NetDormancy == DORM_Initial && In->IsNetStartupActor())
+	if (Actor->NetDormancy == DORM_Initial && Actor->IsNetStartupActor())
 		return false;
 
 /*
 	These checks were done in legacy code and we would like to avoid them.
 
 	// Actors should finish initialization outside of the replication loop. Maybe some weird multi frame delayed case?
-	if (!In->IsActorInitialized())
+	if (!Actor->IsActorInitialized())
 		return false;
 
 	// This check is slow and is not needed unless you are streaming levels on the server. If needed this should be opt in globally some how.
