Remove actors from Replication Graph when world is unloaded

diff --git a/Engine/Plugins/Runtime/ReplicationGraph/Source/Private/ReplicationGraph.cpp b/Engine/Plugins/Runtime/ReplicationGraph/Source/Private/ReplicationGraph.cpp
--- a/Engine/Plugins/Runtime/ReplicationGraph/Source/Private/ReplicationGraph.cpp
+++ b/Engine/Plugins/Runtime/ReplicationGraph/Source/Private/ReplicationGraph.cpp
@@ -335,27 +335,48 @@ UReplicationGraph::UReplicationGraph()
 	}
 #endif
 
-	// Report any actors that were not removed from networking
-	PostWorldCleanupCheckDelegateHandle = FWorldDelegates::OnPostWorldCleanup.AddWeakLambda(
+	// Remove handles to networked actors when their worlds are cleaned up
+	WorldCleanupDelegateHandle = FWorldDelegates::OnWorldCleanup.AddWeakLambda(
 		this,
-		[&ActiveActors = this->ActiveNetworkActors](UWorld*, bool, bool)
+		[this](UWorld* World, bool, bool)
 		{
-			if (!ActiveActors.IsEmpty())
-			{
-				UE_LOG(LogReplicationGraph, Warning, TEXT("%d actors were not removed from ReplicationGraph"), ActiveActors.Num());
+			// Build a list then iterate it to avoid removing while iterating
+			TArray<AActor*> ActorsInWorld;
+			ActorsInWorld.Reserve(ActiveNetworkActors.Num());
 
-				for (AActor* Actor : ActiveActors)
+			for (AActor* Actor : ActiveNetworkActors)
+			{
+				if (Actor->GetWorld() == World)
 				{
-					UE_LOG(LogReplicationGraph, Log, TEXT("  Leaked actor: %s"), *GetFullNameSafe(Actor));
+					ActorsInWorld.Emplace(Actor);
 				}
 			}
+
+			for (AActor* Actor : ActorsInWorld)
+			{
+				RemoveNetworkActor(Actor);
+			}
 		});
 }
 
-UReplicationGraph::~UReplicationGraph()
+void UReplicationGraph::BeginDestroy()
 {
-	FWorldDelegates::OnPostWorldCleanup.Remove(PostWorldCleanupCheckDelegateHandle);
-	PostWorldCleanupCheckDelegateHandle.Reset();
+	FWorldDelegates::OnWorldCleanup.Remove(WorldCleanupDelegateHandle);
+	WorldCleanupDelegateHandle.Reset();
+
+	// Report any leaked actors
+	// Must occur in BeginDestroy() because by FinishDestroy() the actors have been destroyed.
+	if (!ActiveNetworkActors.IsEmpty())
+	{
+		UE_LOG(LogReplicationGraph, Warning, TEXT("%d actors were not removed from ReplicationGraph"), ActiveNetworkActors.Num());
+
+		for (AActor* Actor : ActiveNetworkActors)
+		{
+			UE_LOG(LogReplicationGraph, Log, TEXT("  Leaked actor: %s"), *GetFullNameSafe(Actor));
+		}
+	}
+
+	Super::BeginDestroy();
 }
 
 extern void CountReplicationGraphSharedBytes_Private(FArchive& Ar);
diff --git a/Engine/Plugins/Runtime/ReplicationGraph/Source/Public/ReplicationGraph.h b/Engine/Plugins/Runtime/ReplicationGraph/Source/Public/ReplicationGraph.h
--- a/Engine/Plugins/Runtime/ReplicationGraph/Source/Public/ReplicationGraph.h
+++ b/Engine/Plugins/Runtime/ReplicationGraph/Source/Public/ReplicationGraph.h
@@ -922,7 +922,8 @@ class REPLICATIONGRAPH_API UReplicationGraph : public UReplicationDriver
 public:
 
 	UReplicationGraph();
-	virtual ~UReplicationGraph();
+
+	virtual void BeginDestroy() override;
 
 	/** The per-connection manager class to instantiate. This will be read off the instantiated UNetReplicationManager. */
 	UPROPERTY(Config)
@@ -1223,8 +1224,8 @@ private:
 
 	friend class AReplicationGraphDebugActor;
 
-	/** Delegate that runs after world cleanup to report any actors that were not explicitly removed from Replication Graph and could be leaks. */
-	FDelegateHandle PostWorldCleanupCheckDelegateHandle;
+	/** Delegate to remove all network actors when a world is cleaned up. */
+	FDelegateHandle WorldCleanupDelegateHandle;
 };
 
 // --------------------------------------------------------------------------------------------------------------------------------------------
