Use TVariant instead of std::variant

Removed triviality assertions. TVariant doesn't meet most of them but TArray will treat it as if it does regardless.

#rnx

diff --git a/Engine/Plugins/Runtime/ReplicationGraph/Source/Public/ReplicationGraphTypes.h b/Engine/Plugins/Runtime/ReplicationGraph/Source/Public/ReplicationGraphTypes.h
--- a/Engine/Plugins/Runtime/ReplicationGraph/Source/Public/ReplicationGraphTypes.h
+++ b/Engine/Plugins/Runtime/ReplicationGraph/Source/Public/ReplicationGraphTypes.h
@@ -12,7 +12,6 @@
 #include "UObject/WeakObjectPtr.h"
 
 #include <type_traits>
-#include <variant> // Switching to TVariant once it's trivially destructible
 
 #include "ReplicationGraphTypes.generated.h"
 
@@ -102,9 +101,8 @@ private:
 	// Type is set during construction depending on bUseWeakPointers: FWeakObjectPtr when true, AActor* when false. Most functions should depend on
 	// the actual type, only SetActor() depends on bUseWeakPointers so the configuration can change at runtime without breaking extant instances.
 	//
-	// Requires std::variant rather than TVariant to be trivially destructible for FActorRepList's variable-length FActorRepListType[] Data member.
 	// Could be condensed to sizeof(AActor*) aka sizeof(FWeakObjectPtr) by storing the type flag in the high bit and masking it out.
-	std::variant<AActor*, FWeakObjectPtr> ActorUnion;
+	TVariant<AActor*, FWeakObjectPtr> ActorUnion;
 
 public:
 	/** More info for debugging - ActorRaw's name */
@@ -158,7 +156,7 @@ public:
 	friend bool operator==(const FActorRepListType& Left, const FActorRepListType& Right)
 	{
 		// Avoid FWeakObjectPtr::operator==() because it resolves both weak pointers if they don't match.
-		if (const FWeakObjectPtr* LeftPtr = std::get_if<FWeakObjectPtr>(&Left.ActorUnion), *RightPtr = std::get_if<FWeakObjectPtr>(&Right.ActorUnion);
+		if (const FWeakObjectPtr* LeftPtr = Left.ActorUnion.TryGet<FWeakObjectPtr>(), *RightPtr = Right.ActorUnion.TryGet<FWeakObjectPtr>();
 			LeftPtr && RightPtr)
 		{
 			return LeftPtr->HasSameIndexAndSerialNumber(*RightPtr);
@@ -193,12 +191,12 @@ public:
 	// comparison with nullptr without resolving weak pointer or implicitly constructing a new FActorRepListType(nullptr).
 	friend bool operator==(const FActorRepListType& RepListActor, nullptr_t)
 	{
-		if (const FWeakObjectPtr* WeakPtrPtr = std::get_if<FWeakObjectPtr>(&RepListActor.ActorUnion); WeakPtrPtr)
+		if (const FWeakObjectPtr* WeakPtrPtr = RepListActor.ActorUnion.TryGet<FWeakObjectPtr>(); WeakPtrPtr)
 		{
 			return WeakPtrPtr->IsExplicitlyNull();
 		}
 
-		AActor* const* RawActorPtr = std::get_if<AActor*>(&RepListActor.ActorUnion);
+		AActor* const* RawActorPtr = RepListActor.ActorUnion.TryGet<AActor*>();
 		return ensure(RawActorPtr) && *RawActorPtr == nullptr;
 	}
 	friend bool operator!=(const FActorRepListType& RepListActor, nullptr_t)
@@ -216,22 +214,22 @@ public:
 
 	bool IsValid() const
 	{
-		if (std::holds_alternative<AActor*>(ActorUnion))
+		if (ActorUnion.IsType<AActor*>())
 		{
-			return std::get<AActor*>(ActorUnion) != nullptr;
+			return ActorUnion.Get<AActor*>() != nullptr;
 		}
 
-		return std::get<FWeakObjectPtr>(ActorUnion).IsValid();
+		return ActorUnion.Get<FWeakObjectPtr>().IsValid();
 	}
 
 	inline AActor* GetActor() const
 	{
-		if (std::holds_alternative<AActor*>(ActorUnion))
+		if (ActorUnion.IsType<AActor*>())
 		{
-			return std::get<AActor*>(ActorUnion);
+			return ActorUnion.Get<AActor*>();
 		}
 
-		const FWeakObjectPtr& WeakActor = std::get<FWeakObjectPtr>(ActorUnion);
+		const FWeakObjectPtr& WeakActor = ActorUnion.Get<FWeakObjectPtr>();
 
 		if (WeakActor.IsExplicitlyNull())
 		{
@@ -264,11 +262,11 @@ private:
 	{
 		if (UE::Net::RepGraph::bUseWeakPointers)
 		{
-			ActorUnion.emplace<FWeakObjectPtr>(Actor);
+			ActorUnion.Emplace<FWeakObjectPtr>(Actor);
 		}
 		else
 		{
-			ActorUnion.emplace<AActor*>(Actor);
+			ActorUnion.Emplace<AActor*>(Actor);
 		}
 
 		SetDebugInfo();
@@ -302,13 +300,6 @@ template< class T > FORCEINLINE T* CastChecked(const FActorRepListType& Src, ECa
 }
 #endif // UE_ACTOR_REPLIST_TYPE_EXTRA_SAFETY
 
-// Ensure desired performance characteristics.
-static_assert(std::is_trivially_destructible_v<FActorRepListType>);
-static_assert(std::is_trivially_move_constructible_v<FActorRepListType>);
-static_assert(std::is_trivially_move_assignable_v<FActorRepListType>);
-static_assert(std::is_trivially_copy_constructible_v<FActorRepListType>);
-static_assert(std::is_trivially_copy_assignable_v<FActorRepListType>);
-
 FORCEINLINE FString GetActorRepListTypeDebugString(const FActorRepListType& In) { return GetNameSafe(In); }
 FORCEINLINE UClass* GetActorRepListTypeClass(const FActorRepListType& In) { return In->GetClass(); }
 
