Replace some usages of FORCEINLINE with inline in Server modules.

For context about this change, see this doc: https://docs.google.com/document/d/1LjpT7NpzTecfXdFNyv5KM6sCu67827BRljNi1iaWVSE/edit?tab=t.0#heading=h.fvtz8gjukz7x

Includes modules: ReplicationGraph, IrisCore, NetCore, NetworkReplayStreaming, PacketHandler, ReliabilityHandlerComponent, Sockets

Reviewers were automatically determined. Please comment if anyone else should be added.

#rnx

diff --git a/Engine/Plugins/Runtime/ReplicationGraph/Source/Public/ReplicationGraphTypes.h b/Engine/Plugins/Runtime/ReplicationGraph/Source/Public/ReplicationGraphTypes.h
--- a/Engine/Plugins/Runtime/ReplicationGraph/Source/Public/ReplicationGraphTypes.h
+++ b/Engine/Plugins/Runtime/ReplicationGraph/Source/Public/ReplicationGraphTypes.h
@@ -294,16 +294,16 @@ private:
 	}
 };
 
-template< class T > FORCEINLINE T* Cast(const FActorRepListType& Src) { return Cast<T>(Src.GetActor()); }
-template< class T > FORCEINLINE T* ExactCast(const FActorRepListType& Src) { return ExactCast<T>(Src.GetActor()); }
-template< class T > FORCEINLINE T* CastChecked(const FActorRepListType& Src, ECastCheckedType::Type CheckType = ECastCheckedType::NullChecked)
+template< class T > inline T* Cast(const FActorRepListType& Src) { return Cast<T>(Src.GetActor()); }
+template< class T > inline T* ExactCast(const FActorRepListType& Src) { return ExactCast<T>(Src.GetActor()); }
+template< class T > inline T* CastChecked(const FActorRepListType& Src, ECastCheckedType::Type CheckType = ECastCheckedType::NullChecked)
 {
 	return CastChecked<T>(Src.GetActor(), CheckType);
 }
 #endif // UE_ACTOR_REPLIST_TYPE_EXTRA_SAFETY
 
-FORCEINLINE FString GetActorRepListTypeDebugString(const FActorRepListType& In) { return GetNameSafe(In); }
-FORCEINLINE UClass* GetActorRepListTypeClass(const FActorRepListType& In) { return In->GetClass(); }
+inline FString GetActorRepListTypeDebugString(const FActorRepListType& In) { return GetNameSafe(In); }
+inline UClass* GetActorRepListTypeClass(const FActorRepListType& In) { return In->GetClass(); }
 
 // Generic flags that describe an actor list. Currently only used for "Default" vs "FastShared" path. These are for unsorted lists that are returned from the graph, merged together, sorted, and finally replicated until bandwidth limit is hit
 enum class EActorRepListTypeFlags : uint8
@@ -314,17 +314,17 @@ enum class EActorRepListTypeFlags : uint8
 };
 
 // Tests if an actor is valid for replication: not pending kill, etc. Says nothing about wanting to replicate or should replicate, etc.
-FORCEINLINE bool IsActorValidForReplication(const AActor* Actor)
+inline bool IsActorValidForReplication(const AActor* Actor)
 {
 	return Actor && !Actor->IsActorBeingDestroyed() && IsValidChecked(Actor) && !Actor->IsUnreachable(); 
 }
-FORCEINLINE bool IsActorValidForReplication(const FActorRepListType& In)
+inline bool IsActorValidForReplication(const FActorRepListType& In)
 {
 	return IsActorValidForReplication(static_cast<const AActor*>(In));
 }
 
 REPLICATIONGRAPH_API void LogMoreInfoOnIsActorValidFailure(const FActorRepListType& In);
-FORCEINLINE bool IsActorValidForReplication_LogMoreInfo(const FActorRepListType& In)
+inline bool IsActorValidForReplication_LogMoreInfo(const FActorRepListType& In)
 { 
 	if (LIKELY(IsActorValidForReplication(In)))
 	{
@@ -336,7 +336,7 @@ FORCEINLINE bool IsActorValidForReplication_LogMoreInfo(const FActorRepListType&
 }
 
 // Tests if an actor is valid for replication gathering. Meaning, it can be gathered from the replication graph and considered for replication.
-FORCEINLINE bool IsActorValidForReplicationGather(const FActorRepListType& In)
+inline bool IsActorValidForReplicationGather(const FActorRepListType& In)
 {
 	const AActor* Actor = In;
 
@@ -403,12 +403,12 @@ struct TActorRepListViewBase
 	TActorRepListViewBase() { }
 	TActorRepListViewBase(const PointerType& In) : RepList(In) { }
 
-	FORCEINLINE int32 Num() const { return RepList->Num; }
-	FORCEINLINE const FActorRepListType& operator[](int32 idx) const  { repCheck(RepList); repCheck(RepList->Max > idx); return RepList->Data[idx]; }
+	inline int32 Num() const { return RepList->Num; }
+	inline const FActorRepListType& operator[](int32 idx) const  { repCheck(RepList); repCheck(RepList->Max > idx); return RepList->Data[idx]; }
 
 	/** Resets the view to null - meaning it is not pointing to any list (as opposed to an empty list of some preallocated size) */
-	FORCEINLINE void ResetToNull() { RepList = nullptr; }
-	FORCEINLINE bool IsValid() const { return RepList != nullptr; }
+	inline void ResetToNull() { RepList = nullptr; }
+	inline bool IsValid() const { return RepList != nullptr; }
 
 	int32 IndexOf(const FActorRepListType& Value) const
 	{
@@ -471,8 +471,8 @@ struct TActorRepListViewBase
 
 private:
 
-	FORCEINLINE friend FActorRepListType* begin(const TActorRepListViewBase<PointerType>& View) { return View.RepList->Data; }
-	FORCEINLINE friend FActorRepListType* end(const TActorRepListViewBase<PointerType>& View) { return View.RepList->Data + View.RepList->Num; }
+	inline friend FActorRepListType* begin(const TActorRepListViewBase<PointerType>& View) { return View.RepList->Data; }
+	inline friend FActorRepListType* end(const TActorRepListViewBase<PointerType>& View) { return View.RepList->Data + View.RepList->Num; }
 };
 
 /**
@@ -517,7 +517,7 @@ struct FActorRepListRefView
 		return RepList.RemoveSingleSwap(ElementToRemove, AllowShrink) > 0;
 	}
 	UE_ALLOWSHRINKING_BOOL_DEPRECATED("RemoveFast")
-	FORCEINLINE bool RemoveFast(const FActorRepListType& ElementToRemove, bool bAllowShrink)
+	inline bool RemoveFast(const FActorRepListType& ElementToRemove, bool bAllowShrink)
 	{
 		return RemoveFast(ElementToRemove, bAllowShrink ? EAllowShrinking::Yes : EAllowShrinking::No);
 	}
@@ -663,29 +663,29 @@ struct FGatheredReplicationActorLists
 		List.AppendToTArray(ReplicationLists[(uint32)Flags]);
 	}
 
-	FORCEINLINE void Reset()
+	inline void Reset()
 	{
 		for (uint32 i = (uint32)EActorRepListTypeFlags::Default; i < (uint32)EActorRepListTypeFlags::Max; ++i)
 		{
 			ReplicationLists[i].Reset();
 		}
 	}
-	FORCEINLINE int32 NumLists() const
+	inline int32 NumLists() const
 	{
 		return ReplicationLists.Num();
 	}
 
 PRAGMA_DISABLE_DEPRECATION_WARNINGS
 	UE_DEPRECATED(5.5, "Use ViewActors() instead")
-	UE_API FORCEINLINE const TArray<FActorRepListConstView>& GetLists(EActorRepListTypeFlags ListFlags) const;
+	UE_API inline const TArray<FActorRepListConstView>& GetLists(EActorRepListTypeFlags ListFlags) const;
 PRAGMA_ENABLE_DEPRECATION_WARNINGS
 
-	FORCEINLINE TArrayView<const FActorRepListType> ViewActors(EActorRepListTypeFlags ListFlags) const
+	inline TArrayView<const FActorRepListType> ViewActors(EActorRepListTypeFlags ListFlags) const
 	{
 		return ReplicationLists[(uint32)ListFlags];
 	}
 
-	FORCEINLINE bool ContainsLists(EActorRepListTypeFlags Flags) const
+	inline bool ContainsLists(EActorRepListTypeFlags Flags) const
 	{
 		return ReplicationLists[(uint32)Flags].Num() > 0;
 	}
@@ -1161,8 +1161,8 @@ struct TClassMap
 		Map.Emplace(FObjectKey(InClass), MoveTemp(Value));
 	}
 
-	FORCEINLINE typename TMap<FObjectKey, ValueType>::TIterator CreateIterator() { return Map.CreateIterator(); }
-	FORCEINLINE void Reset() { Map.Reset(); }
+	inline typename TMap<FObjectKey, ValueType>::TIterator CreateIterator() { return Map.CreateIterator(); }
+	inline void Reset() { Map.Reset(); }
 
 	TFunction<bool(UClass*, ValueType&)>	InitNewElement;
 
@@ -1255,7 +1255,7 @@ struct FGlobalActorReplicationInfoMap
 
 
 	/** Finds data associated with the actor but does not create if its not there yet. */
-	FORCEINLINE FGlobalActorReplicationInfo* Find(const FActorRepListType& Actor) const 
+	inline FGlobalActorReplicationInfo* Find(const FActorRepListType& Actor) const 
 	{
 		if (const TUniquePtr<FGlobalActorReplicationInfo>* Ptr = ActorMap.Find(Actor))
 		{
@@ -1272,18 +1272,18 @@ struct FGlobalActorReplicationInfoMap
 	void NotifyActorRenamed(AActor* Actor, FName PreviousStreamingLevelName);
 	
 	/** Returns ClassInfo for a given class. */
-	FORCEINLINE FClassReplicationInfo& GetClassInfo(UClass* Class) { return ClassMap.GetChecked(Class); }
+	inline FClassReplicationInfo& GetClassInfo(UClass* Class) { return ClassMap.GetChecked(Class); }
 
 	/** Sets class info for a given class and its derived classes if desired. Call this in your Replication Graphs setup */
-	FORCEINLINE void SetClassInfo(UClass* InClass, const FClassReplicationInfo& Info) {	ClassMap.Set(InClass, Info); }
+	inline void SetClassInfo(UClass* InClass, const FClassReplicationInfo& Info) {	ClassMap.Set(InClass, Info); }
 	
-	FORCEINLINE TMap<FActorRepListType, TUniquePtr<FGlobalActorReplicationInfo>>::TIterator CreateActorMapIterator() { return ActorMap.CreateIterator(); }
-	FORCEINLINE TMap<FActorRepListType, TUniquePtr<FGlobalActorReplicationInfo>>::TConstIterator CreateActorMapIterator() const { return ActorMap.CreateConstIterator(); }
-	FORCEINLINE TMap<FObjectKey, FClassReplicationInfo>::TIterator CreateClassMapIterator() { return ClassMap.CreateIterator(); }
+	inline TMap<FActorRepListType, TUniquePtr<FGlobalActorReplicationInfo>>::TIterator CreateActorMapIterator() { return ActorMap.CreateIterator(); }
+	inline TMap<FActorRepListType, TUniquePtr<FGlobalActorReplicationInfo>>::TConstIterator CreateActorMapIterator() const { return ActorMap.CreateConstIterator(); }
+	inline TMap<FObjectKey, FClassReplicationInfo>::TIterator CreateClassMapIterator() { return ClassMap.CreateIterator(); }
 
 	int32 Num() const { return ActorMap.Num(); }
 
-	FORCEINLINE void ResetActorMap() { ActorMap.Reset(); }
+	inline void ResetActorMap() { ActorMap.Reset(); }
 
 	void CountBytes(FArchive& Ar) const
 	{
@@ -1474,7 +1474,7 @@ public:
 /** Map for Actor -> ConnectionActorInfo. This wraps the TMap mainly so we can do custom initialization in FindOrAdd. */
 struct FPerConnectionActorInfoMap
 {
-	FORCEINLINE_DEBUGGABLE FConnectionReplicationActorInfo& FindOrAdd(const FActorRepListType& Actor)
+	inline FConnectionReplicationActorInfo& FindOrAdd(const FActorRepListType& Actor)
 	{
 		if (TSharedPtr<FConnectionReplicationActorInfo>* ValuePtr = ActorMap.Find(Actor))
 		{
@@ -1486,7 +1486,7 @@ struct FPerConnectionActorInfoMap
 		return *NewInfo;
 	}
 
-	FORCEINLINE FConnectionReplicationActorInfo* Find(const FActorRepListType& Actor)
+	inline FConnectionReplicationActorInfo* Find(const FActorRepListType& Actor)
 	{
 		if (TSharedPtr<FConnectionReplicationActorInfo>* ValuePtr = ActorMap.Find(Actor))
 		{
@@ -1496,7 +1496,7 @@ struct FPerConnectionActorInfoMap
 		return nullptr;
 	}
 
-	FORCEINLINE FConnectionReplicationActorInfo* FindByChannel(UActorChannel* Channel)
+	inline FConnectionReplicationActorInfo* FindByChannel(UActorChannel* Channel)
 	{
 		if (TSharedPtr<FConnectionReplicationActorInfo>* ValuePtr = ChannelMap.Find(Channel))
 		{
@@ -1506,7 +1506,7 @@ struct FPerConnectionActorInfoMap
 		return nullptr;
 	}
 
-	FORCEINLINE void AddChannel(const FActorRepListType& Actor, UActorChannel* Channel)
+	inline void AddChannel(const FActorRepListType& Actor, UActorChannel* Channel)
 	{
 		if (TSharedPtr<FConnectionReplicationActorInfo>* ValuePtr = ActorMap.Find(Actor))
 		{
@@ -1514,27 +1514,27 @@ struct FPerConnectionActorInfoMap
 		}
 	}
 
-	FORCEINLINE void RemoveChannel(UActorChannel* Channel)
+	inline void RemoveChannel(UActorChannel* Channel)
 	{
 		ChannelMap.Remove(Channel);
 	}
 
-	FORCEINLINE void RemoveActor(const FActorRepListType& Actor)
+	inline void RemoveActor(const FActorRepListType& Actor)
 	{
 		ActorMap.Remove(Actor);
 	}
 
-	FORCEINLINE TMap<FActorRepListType, TSharedPtr<FConnectionReplicationActorInfo>>::TIterator CreateIterator()
+	inline TMap<FActorRepListType, TSharedPtr<FConnectionReplicationActorInfo>>::TIterator CreateIterator()
 	{
 		return ActorMap.CreateIterator();
 	}
 
-	FORCEINLINE void SetGlobalMap(FGlobalActorReplicationInfoMap* InGlobalMap)
+	inline void SetGlobalMap(FGlobalActorReplicationInfoMap* InGlobalMap)
 	{
 		GlobalMap = InGlobalMap;
 	}
 
-	FORCEINLINE TMap<UActorChannel*, TSharedPtr<FConnectionReplicationActorInfo>>::TIterator CreateChannelIterator()
+	inline TMap<UActorChannel*, TSharedPtr<FConnectionReplicationActorInfo>>::TIterator CreateChannelIterator()
 	{
 		return ChannelMap.CreateIterator();
 	}
@@ -1881,7 +1881,7 @@ struct FNativeClassAccumulator
 #if WITH_EDITOR
 void ForEachClientPIEWorld(TFunction<void(UWorld*)> Func);
 #else
-FORCEINLINE void ForEachClientPIEWorld(TFunction<void(UWorld*)> Func) { }
+inline void ForEachClientPIEWorld(TFunction<void(UWorld*)> Func) { }
 #endif
 
 // --------------------------------------------------------------------------------------------------------------------------------------------
diff --git a/Engine/Source/Runtime/Experimental/Iris/Core/Public/Iris/ReplicationSystem/NetObjectGroupHandle.h b/Engine/Source/Runtime/Experimental/Iris/Core/Public/Iris/ReplicationSystem/NetObjectGroupHandle.h
--- a/Engine/Source/Runtime/Experimental/Iris/Core/Public/Iris/ReplicationSystem/NetObjectGroupHandle.h
+++ b/Engine/Source/Runtime/Experimental/Iris/Core/Public/Iris/ReplicationSystem/NetObjectGroupHandle.h
@@ -105,7 +105,7 @@ private:
 
 static_assert(sizeof(FNetObjectGroupHandle) == sizeof(uint64), "FNetObjectGroupHandle must be of size 64bits.");
 
-FORCEINLINE uint32 GetTypeHash(const FNetObjectGroupHandle Handle)
+inline uint32 GetTypeHash(const FNetObjectGroupHandle Handle)
 {
 	return ::GetTypeHash(Handle.GetRawValue());
 }
diff --git a/Engine/Source/Runtime/Experimental/Iris/Core/Public/Iris/ReplicationSystem/NetRefHandle.h b/Engine/Source/Runtime/Experimental/Iris/Core/Public/Iris/ReplicationSystem/NetRefHandle.h
--- a/Engine/Source/Runtime/Experimental/Iris/Core/Public/Iris/ReplicationSystem/NetRefHandle.h
+++ b/Engine/Source/Runtime/Experimental/Iris/Core/Public/Iris/ReplicationSystem/NetRefHandle.h
@@ -78,12 +78,12 @@ private:
 	};
 };
 
-FORCEINLINE uint32 GetTypeHash(const FNetRefHandle& Handle)
+inline uint32 GetTypeHash(const FNetRefHandle& Handle)
 {
 	return ::GetTypeHash(Handle.GetId());
 }
 
-FORCEINLINE uint64 GetObjectIdForNetTrace(const FNetRefHandle& Handle)
+inline uint64 GetObjectIdForNetTrace(const FNetRefHandle& Handle)
 {
 	return Handle.GetId();
 }
diff --git a/Engine/Source/Runtime/Experimental/Iris/Core/Public/Iris/Serialization/PolymorphicNetSerializerImpl.h b/Engine/Source/Runtime/Experimental/Iris/Core/Public/Iris/Serialization/PolymorphicNetSerializerImpl.h
--- a/Engine/Source/Runtime/Experimental/Iris/Core/Public/Iris/Serialization/PolymorphicNetSerializerImpl.h
+++ b/Engine/Source/Runtime/Experimental/Iris/Core/Public/Iris/Serialization/PolymorphicNetSerializerImpl.h
@@ -184,7 +184,7 @@ struct TPolymorphicArrayStructNetSerializerImpl : protected Private::FPolymorphi
 
 	struct FSourceArrayItemTypeDeleter
 	{
-		FORCEINLINE void operator()(SourceArrayItemType* Object) const
+		inline void operator()(SourceArrayItemType* Object) const
 		{
 			check(Object);
 			UScriptStruct* ScriptStruct = Object->GetScriptStruct();
diff --git a/Engine/Source/Runtime/Net/Core/Classes/Net/Serialization/FastArraySerializer.h b/Engine/Source/Runtime/Net/Core/Classes/Net/Serialization/FastArraySerializer.h
--- a/Engine/Source/Runtime/Net/Core/Classes/Net/Serialization/FastArraySerializer.h
+++ b/Engine/Source/Runtime/Net/Core/Classes/Net/Serialization/FastArraySerializer.h
@@ -333,7 +333,7 @@ struct FFastArraySerializerItem
 	 * 
 	 * NOTE: intentionally not virtual; invoked via templated code, @see FExampleItemEntry
 	 */
-	FORCEINLINE void PreReplicatedRemove(const struct FFastArraySerializer& InArraySerializer) { }
+	inline void PreReplicatedRemove(const struct FFastArraySerializer& InArraySerializer) { }
 	/**
 	 * Called after adding and serializing a new element
 	 *
@@ -341,14 +341,14 @@ struct FFastArraySerializerItem
 	 * 
 	 * NOTE: intentionally not virtual; invoked via templated code, @see FExampleItemEntry
 	 */
-	FORCEINLINE void PostReplicatedAdd(const struct FFastArraySerializer& InArraySerializer) { }
+	inline void PostReplicatedAdd(const struct FFastArraySerializer& InArraySerializer) { }
 	/**
 	 * Called after updating an existing element with new data
 	 *
 	 * @param InArraySerializer	Array serializer that owns the item and has triggered the replication call
 	 * NOTE: intentionally not virtual; invoked via templated code, @see FExampleItemEntry
 	 */
-	FORCEINLINE void PostReplicatedChange(const struct FFastArraySerializer& InArraySerializer) { }
+	inline void PostReplicatedChange(const struct FFastArraySerializer& InArraySerializer) { }
 
 	/**
 	 * Called when logging LogNetFastTArray (log or lower verbosity)
@@ -356,7 +356,7 @@ struct FFastArraySerializerItem
 	 * @param InArraySerializer	Array serializer that owns the item and has triggered the replication call
 	 * NOTE: intentionally not virtual; invoked via templated code, @see FExampleItemEntry
 	 */
-	FORCEINLINE FString GetDebugString() { return FString(TEXT("")); }
+	inline FString GetDebugString() { return FString(TEXT("")); }
 };
 
 /** Struct for holding guid references */
@@ -493,21 +493,21 @@ struct FFastArraySerializer
 	 *
 	 * NOTE: intentionally not virtual; invoked via templated code, @see FExampleItemEntry
 	 */
-	FORCEINLINE void PreReplicatedRemove(const TArrayView<int32>& RemovedIndices, int32 FinalSize) { }
+	inline void PreReplicatedRemove(const TArrayView<int32>& RemovedIndices, int32 FinalSize) { }
 
 	/**
 	 * Called after adding all new elements and after the elements themselves are notified.  The indices are valid for this function call only!
 	 *
 	 * NOTE: intentionally not virtual; invoked via templated code, @see FExampleItemEntry
 	 */
-	FORCEINLINE void PostReplicatedAdd(const TArrayView<int32>& AddedIndices, int32 FinalSize) { }
+	inline void PostReplicatedAdd(const TArrayView<int32>& AddedIndices, int32 FinalSize) { }
 
 	/**
 	 * Called after updating all existing elements with new data and after the elements themselves are notified. The indices are valid for this function call only!
 	 *
 	 * NOTE: intentionally not virtual; invoked via templated code, @see FExampleItemEntry
 	 */
-	FORCEINLINE void PostReplicatedChange(const TArrayView<int32>& ChangedIndices, int32 FinalSize) { }
+	inline void PostReplicatedChange(const TArrayView<int32>& ChangedIndices, int32 FinalSize) { }
 
 	/**
 	 * If a function with the signature void PostReplicatedReceive(const FFastArraySerializer::FPostReplicatedReceiveParameters& Parameters) is defined in the derived struct
diff --git a/Engine/Source/Runtime/Net/Core/Public/Net/Core/Analytics/NetAnalytics.h b/Engine/Source/Runtime/Net/Core/Public/Net/Core/Analytics/NetAnalytics.h
--- a/Engine/Source/Runtime/Net/Core/Public/Net/Core/Analytics/NetAnalytics.h
+++ b/Engine/Source/Runtime/Net/Core/Public/Net/Core/Analytics/NetAnalytics.h
@@ -305,7 +305,7 @@ protected:
  */
 struct FNetAnalyticsDataDeleter
 {
-	FORCEINLINE void operator()(FNetAnalyticsData* Data)
+	inline void operator()(FNetAnalyticsData* Data)
 	{
 		if (Data->Aggregator != nullptr)
 		{
diff --git a/Engine/Source/Runtime/Net/Core/Public/Net/Core/Misc/PacketAudit.h b/Engine/Source/Runtime/Net/Core/Public/Net/Core/Misc/PacketAudit.h
--- a/Engine/Source/Runtime/Net/Core/Public/Net/Core/Misc/PacketAudit.h
+++ b/Engine/Source/Runtime/Net/Core/Public/Net/Core/Misc/PacketAudit.h
@@ -124,7 +124,7 @@ public:
 	 * @param OutPacket		The packet the data is being written to - used for gathering audit information
 	 * @param bByteAligned	Whether or not to treat the packet as if it is byte aligned (used in rare circumstances)
 	 */
-	static FORCEINLINE void AddStage(FString StageName, FBitWriter& OutPacket, bool bByteAligned=false)
+	static inline void AddStage(FString StageName, FBitWriter& OutPacket, bool bByteAligned=false)
 	{
 #if !UE_BUILD_SHIPPING
 		if (GPacketAuditor != nullptr)
@@ -141,7 +141,7 @@ public:
 	 * @param InPacket		The packet the data is being read from - used for gathering audit information
 	 * @param bByteAligned	Whether or not to treat the packet as if it is byte aligned (used in rare circumstances)
 	 */
-	static FORCEINLINE void CheckStage(FString StageName, FBitReader& InPacket, bool bByteAligned=false)
+	static inline void CheckStage(FString StageName, FBitReader& InPacket, bool bByteAligned=false)
 	{
 #if !UE_BUILD_SHIPPING
 		if (GPacketAuditor != nullptr)
@@ -156,7 +156,7 @@ public:
 	 *
 	 * @param OutPacket		The current outgoing packet
 	 */
-	static FORCEINLINE void NotifyLowLevelSend(FBitWriter& OutPacket)
+	static inline void NotifyLowLevelSend(FBitWriter& OutPacket)
 	{
 #if !UE_BUILD_SHIPPING
 		if (GPacketAuditor != nullptr)
@@ -171,7 +171,7 @@ public:
 	 *
 	 * @param InPacket		The current incoming packet
 	 */
-	static FORCEINLINE void NotifyLowLevelReceive(FBitReader& InPacket)
+	static inline void NotifyLowLevelReceive(FBitReader& InPacket)
 	{
 #if !UE_BUILD_SHIPPING
 		if (GPacketAuditor != nullptr)
@@ -183,7 +183,7 @@ public:
 
 	// @todo #JohnB: Deprecate, once the pipeline uses FBitReader/Writers all the way
 #if 1
-	static FORCEINLINE void NotifyLowLevelSend(uint8* Data, uint32 DataLen, uint32 DataLenBits)
+	static inline void NotifyLowLevelSend(uint8* Data, uint32 DataLen, uint32 DataLenBits)
 	{
 #if !UE_BUILD_SHIPPING
 		if (GPacketAuditor != nullptr)
@@ -197,7 +197,7 @@ public:
 #endif
 	}
 
-	static FORCEINLINE void NotifyLowLevelReceive(uint8* Data, uint32 DataLen)
+	static inline void NotifyLowLevelReceive(uint8* Data, uint32 DataLen)
 	{
 #if !UE_BUILD_SHIPPING
 		if (GPacketAuditor != nullptr)
diff --git a/Engine/Source/Runtime/Net/Core/Public/Net/Core/NetBitArray.h b/Engine/Source/Runtime/Net/Core/Public/Net/Core/NetBitArray.h
--- a/Engine/Source/Runtime/Net/Core/Public/Net/Core/NetBitArray.h
+++ b/Engine/Source/Runtime/Net/Core/Public/Net/Core/NetBitArray.h
@@ -127,7 +127,7 @@ public:
 	 * Returns a pointer to the internal storage but validates that the array buffer is the expected size.
 	 * Use this before directly iterating over one or multiple buffers and ensure they all support the max iteration index.
 	 */
-	FORCEINLINE const StorageWordType* GetDataChecked(uint32 MaxWordIterationCount) const UE_LIFETIMEBOUND
+	inline const StorageWordType* GetDataChecked(uint32 MaxWordIterationCount) const UE_LIFETIMEBOUND
 	{
 		checkf((int32)MaxWordIterationCount <= Storage.Num(), TEXT("MaxWordIterationCount (%u) is higher than this NetBitArray word count (%d). Out of Bounds memory access will occur."), MaxWordIterationCount, Storage.Num());
 		return Storage.GetData();
@@ -137,27 +137,27 @@ public:
 	 * Returns a pointer to the internal storage but validates that the array buffer is the expected size.
 	 * Use this before directly iterating over one or multiple buffers and ensure they all support the max iteration index.
 	 */
-	FORCEINLINE StorageWordType* GetDataChecked(uint32 MaxWordIterationCount) UE_LIFETIMEBOUND
+	inline StorageWordType* GetDataChecked(uint32 MaxWordIterationCount) UE_LIFETIMEBOUND
 	{
 		checkf((int32)MaxWordIterationCount <= Storage.Num(), TEXT("MaxWordIterationCount (%u) is higher than this NetBitArray word count (%d). Out of Bounds memory access will occur."), MaxWordIterationCount, Storage.Num());
 		return Storage.GetData();
 	}
 
 	/** Returns a pointer to the internal storage. It's safer to use GetWord() or GetDataChecked() instead. */
-	FORCEINLINE const StorageWordType* GetData() const UE_LIFETIMEBOUND { return Storage.GetData(); }
+	inline const StorageWordType* GetData() const UE_LIFETIMEBOUND { return Storage.GetData(); }
 
 	/** Returns a pointer to the internal storage. It's safer to use GetWord() or GetDataChecked() instead. */
-	FORCEINLINE StorageWordType* GetData() UE_LIFETIMEBOUND { return Storage.GetData(); }
+	inline StorageWordType* GetData() UE_LIFETIMEBOUND { return Storage.GetData(); }
 
 	/** Read an entire word. Useful when you want to read 'WordBitCount(32)' bits at once */
-	FORCEINLINE StorageWordType GetWord(uint32 WordIndex) const UE_LIFETIMEBOUND 
+	inline StorageWordType GetWord(uint32 WordIndex) const UE_LIFETIMEBOUND 
 	{ 
 		checkf(WordIndex < (uint32)Storage.Num(), TEXT("NetBitArray index out of bounds: %u / max %u"), WordIndex, Storage.Num());
 		return Storage.GetData()[WordIndex]; 
 	}
 
 	/** Access an entire word. Useful when you want to write 'WordBitCount(32)' bits at once */
-	FORCEINLINE StorageWordType& GetWord(uint32 WordIndex) UE_LIFETIMEBOUND 
+	inline StorageWordType& GetWord(uint32 WordIndex) UE_LIFETIMEBOUND 
 	{  
 		checkf(WordIndex < (uint32)Storage.Num(), TEXT("NetBitArray index out of bounds: %u / max %u"), WordIndex, Storage.Num());
 		return Storage.GetData()[WordIndex]; 
@@ -427,7 +427,7 @@ public:
 	 * Returns a pointer to the internal storage but validates that the array buffer is the expected size.
 	 * Use this before directly iterating over one or multiple buffers and ensure they all support the max iteration index.
 	 */
-	FORCEINLINE const StorageWordType* GetDataChecked(uint32 MaxWordIterationCount) const UE_LIFETIMEBOUND
+	inline const StorageWordType* GetDataChecked(uint32 MaxWordIterationCount) const UE_LIFETIMEBOUND
 	{
 		checkf(MaxWordIterationCount <= WordCount, TEXT("MaxWordIterationCount (%u) is higher than this NetBitArrayView word count (%u). Out of Bounds memory access will occur."), MaxWordIterationCount, WordCount);
 		return Storage;
@@ -437,27 +437,27 @@ public:
 	 * Returns a pointer to the internal storage but validates that the array buffer is the expected size.
 	 * Use this before directly iterating over one or multiple buffers and ensure they all support the max iteration index.
 	 */
-	FORCEINLINE StorageWordType* GetDataChecked(uint32 MaxWordIterationCount) UE_LIFETIMEBOUND
+	inline StorageWordType* GetDataChecked(uint32 MaxWordIterationCount) UE_LIFETIMEBOUND
 	{
 		checkf(MaxWordIterationCount <= WordCount, TEXT("MaxWordIterationCount (%u) is higher than this NetBitArrayView word count (%u). Out of Bounds memory access will occur."), MaxWordIterationCount, WordCount);
 		return Storage;
 	}
 
 	/** Returns a pointer to the internal storage. It's safer to use GetWord() or GetDataChecked() instead. */
-	FORCEINLINE const StorageWordType* GetData() const UE_LIFETIMEBOUND { return Storage; }
+	inline const StorageWordType* GetData() const UE_LIFETIMEBOUND { return Storage; }
 
 	/** Returns a pointer to the internal storage. It's safer to use GetWord() or GetDataChecked() instead. */
-	FORCEINLINE StorageWordType* GetData() UE_LIFETIMEBOUND { return Storage; }
+	inline StorageWordType* GetData() UE_LIFETIMEBOUND { return Storage; }
 
 	/** Read an entire word. Useful when you want to read 'WordBitCount(32)' bits at once */
-	FORCEINLINE StorageWordType GetWord(uint32 WordIndex) const UE_LIFETIMEBOUND 
+	inline StorageWordType GetWord(uint32 WordIndex) const UE_LIFETIMEBOUND 
 	{ 
 		checkf(WordIndex < WordCount, TEXT("NetBitArrayView index out of bounds: %u / max %u"), WordIndex, WordCount);
 		return Storage[WordIndex]; 
 	}
 
 	/** Access an entire word. Useful when you want to write 'WordBitCount(32)' bits at once */
-	FORCEINLINE StorageWordType& GetWord(uint32 WordIndex) UE_LIFETIMEBOUND 
+	inline StorageWordType& GetWord(uint32 WordIndex) UE_LIFETIMEBOUND 
 	{
 		checkf(WordIndex < WordCount, TEXT("NetBitArrayView index out of bounds: %u / max %u"), WordIndex, WordCount);
 		return Storage[WordIndex]; 
diff --git a/Engine/Source/Runtime/NetworkReplayStreaming/NetworkReplayStreaming/Public/NetworkReplayStreaming.h b/Engine/Source/Runtime/NetworkReplayStreaming/NetworkReplayStreaming/Public/NetworkReplayStreaming.h
--- a/Engine/Source/Runtime/NetworkReplayStreaming/NetworkReplayStreaming/Public/NetworkReplayStreaming.h
+++ b/Engine/Source/Runtime/NetworkReplayStreaming/NetworkReplayStreaming/Public/NetworkReplayStreaming.h
@@ -232,7 +232,7 @@ struct FStartStreamingResult : public FStreamingResultBase
 
 DECLARE_DELEGATE_OneParam(FStartStreamingCallback, const FStartStreamingResult&);
 
-static FORCEINLINE FStartStreamingCallback UpgradeStartStreamingDelegate(const FOnStreamReadyDelegate& OldDelegate)
+static inline FStartStreamingCallback UpgradeStartStreamingDelegate(const FOnStreamReadyDelegate& OldDelegate)
 {
 	if (OldDelegate.IsBound())
 	{
@@ -276,7 +276,7 @@ struct FGotoResult : public FStreamingResultBase
 
 DECLARE_DELEGATE_OneParam(FGotoCallback, const FGotoResult&);
 
-static FORCEINLINE FGotoCallback UpgradeGotoDelegate(const FOnCheckpointReadyDelegate& OldDelegate)
+static inline FGotoCallback UpgradeGotoDelegate(const FOnCheckpointReadyDelegate& OldDelegate)
 {
 	if (OldDelegate.IsBound())
 	{
@@ -299,7 +299,7 @@ struct FDeleteFinishedStreamResult : public FStreamingResultBase
 
 DECLARE_DELEGATE_OneParam(FDeleteFinishedStreamCallback, const FDeleteFinishedStreamResult&);
 
-static FORCEINLINE FDeleteFinishedStreamCallback UpgradeDeleteFinishedStreamDelegate(const FOnDeleteFinishedStreamComplete& OldDelegate)
+static inline FDeleteFinishedStreamCallback UpgradeDeleteFinishedStreamDelegate(const FOnDeleteFinishedStreamComplete& OldDelegate)
 {
 	if (OldDelegate.IsBound())
 	{
@@ -327,7 +327,7 @@ struct FEnumerateStreamsResult : public FStreamingResultBase
 
 DECLARE_DELEGATE_OneParam(FEnumerateStreamsCallback, const FEnumerateStreamsResult&);
 
-static FORCEINLINE FEnumerateStreamsCallback UpgradeEnumerateStreamsDelegate(const FOnEnumerateStreamsComplete& OldDelegate)
+static inline FEnumerateStreamsCallback UpgradeEnumerateStreamsDelegate(const FOnEnumerateStreamsComplete& OldDelegate)
 {
 	if (OldDelegate.IsBound())
 	{
@@ -352,7 +352,7 @@ struct FEnumerateEventsResult : public FStreamingResultBase
 
 DECLARE_DELEGATE_OneParam(FEnumerateEventsCallback, const FEnumerateEventsResult&);
 
-static FORCEINLINE FEnumerateEventsCallback UpgradeEnumerateEventsDelegate(const FEnumerateEventsCompleteDelegate& OldDelegate)
+static inline FEnumerateEventsCallback UpgradeEnumerateEventsDelegate(const FEnumerateEventsCompleteDelegate& OldDelegate)
 {
 	if (OldDelegate.IsBound())
 	{
@@ -377,7 +377,7 @@ struct FRequestEventDataResult : public FStreamingResultBase
 
 DECLARE_DELEGATE_OneParam(FRequestEventDataCallback, const FRequestEventDataResult&);
 
-static FORCEINLINE FRequestEventDataCallback UpgradeRequestEventDelegate(const FOnRequestEventDataComplete& OldDelegate)
+static inline FRequestEventDataCallback UpgradeRequestEventDelegate(const FOnRequestEventDataComplete& OldDelegate)
 {
 	if (OldDelegate.IsBound())
 	{
@@ -410,7 +410,7 @@ struct FDownloadHeaderResult : public FStreamingResultBase
 
 DECLARE_DELEGATE_OneParam(FDownloadHeaderCallback, const FDownloadHeaderResult&);
 
-static FORCEINLINE FDownloadHeaderCallback UpgradeDownloadHeaderDelegate(const FOnDownloadHeaderComplete& OldDelegate)
+static inline FDownloadHeaderCallback UpgradeDownloadHeaderDelegate(const FOnDownloadHeaderComplete& OldDelegate)
 {
 	if (OldDelegate.IsBound())
 	{
@@ -438,7 +438,7 @@ struct FSearchEventsResult : public FStreamingResultBase
 
 DECLARE_DELEGATE_OneParam(FSearchEventsCallback, const FSearchEventsResult&);
 
-static FORCEINLINE FSearchEventsCallback UpgradeSearchEventsDelegate(const FOnEnumerateStreamsComplete& OldDelegate)
+static inline FSearchEventsCallback UpgradeSearchEventsDelegate(const FOnEnumerateStreamsComplete& OldDelegate)
 {
 	if (OldDelegate.IsBound())
 	{
diff --git a/Engine/Source/Runtime/PacketHandlers/PacketHandler/Public/PacketHandler.h b/Engine/Source/Runtime/PacketHandlers/PacketHandler/Public/PacketHandler.h
--- a/Engine/Source/Runtime/PacketHandlers/PacketHandler/Public/PacketHandler.h
+++ b/Engine/Source/Runtime/PacketHandlers/PacketHandler/Public/PacketHandler.h
@@ -359,7 +359,7 @@ public:
 
 	// @todo: Don't deprecate, until after the NetDriver refactor
 	//UE_DEPRECATED(4.26, "Incoming now uses FReceivedPacketView.")
-	FORCEINLINE const ProcessedPacket Incoming(uint8* Packet, int32 CountBytes)
+	inline const ProcessedPacket Incoming(uint8* Packet, int32 CountBytes)
 	{
 		PRAGMA_DISABLE_DEPRECATION_WARNINGS
 		const TSharedPtr<const FInternetAddr> EmptyAddress = nullptr;
@@ -376,7 +376,7 @@ public:
 	 * @param PacketView	View of the packet being processed - PacketView.Data should be used to return the result
 	 * @return				Returns Success/Failure
 	 */
-	FORCEINLINE EIncomingResult Incoming(FReceivedPacketView& PacketView)
+	inline EIncomingResult Incoming(FReceivedPacketView& PacketView)
 	{
 		return Incoming_Internal(PacketView);
 	}
@@ -392,7 +392,7 @@ public:
 	 * @param Traits		Traits for the packet, passed down from the NetConnection
 	 * @return				Returns the final packet
 	 */
-	FORCEINLINE const ProcessedPacket Outgoing(uint8* Packet, int32 CountBits, FOutPacketTraits& Traits)
+	inline const ProcessedPacket Outgoing(uint8* Packet, int32 CountBits, FOutPacketTraits& Traits)
 	{
 		const TSharedPtr<const FInternetAddr> EmptyAddress = nullptr;
 		return Outgoing_Internal(Packet, CountBits, Traits, false, EmptyAddress);
@@ -401,7 +401,7 @@ public:
 
 	// @todo: Don't deprecate, until after the NetDriver refactor
 	//UE_DEPRECATED(4.26, "IncomingConnectionless now uses FReceivedPacketView.")
-	FORCEINLINE const ProcessedPacket IncomingConnectionless(const TSharedPtr<const FInternetAddr>& Address, uint8* Packet, int32 CountBytes)
+	inline const ProcessedPacket IncomingConnectionless(const TSharedPtr<const FInternetAddr>& Address, uint8* Packet, int32 CountBytes)
 	{
 		PRAGMA_DISABLE_DEPRECATION_WARNINGS
 		return Incoming_Internal(Packet, CountBytes, true, Address);
@@ -416,7 +416,7 @@ public:
 	 * @param PacketView	View of the packet being processed - PacketView.Data should be used to return the result
 	 * @return				Returns Success/Failure
 	 */
-	FORCEINLINE EIncomingResult IncomingConnectionless(FReceivedPacketView& PacketView)
+	inline EIncomingResult IncomingConnectionless(FReceivedPacketView& PacketView)
 	{
 		PacketView.Traits.bConnectionlessPacket = true;
 
@@ -433,7 +433,7 @@ public:
 	 * @param Traits		Traits for the packet, if applicable
 	 * @return				Returns the final packet
 	 */
-	FORCEINLINE const ProcessedPacket OutgoingConnectionless(const TSharedPtr<const FInternetAddr>& Address, uint8* Packet, int32 CountBits, FOutPacketTraits& Traits)
+	inline const ProcessedPacket OutgoingConnectionless(const TSharedPtr<const FInternetAddr>& Address, uint8* Packet, int32 CountBits, FOutPacketTraits& Traits)
 	{
 		return Outgoing_Internal(Packet, CountBits, Traits, true, Address);
 	}
@@ -512,7 +512,7 @@ public:
 	 *
 	 * @param PacketToQueue		The packet to be queued
 	 */
-	FORCEINLINE void QueuePacketForRawSending(BufferedPacket* PacketToQueue)
+	inline void QueuePacketForRawSending(BufferedPacket* PacketToQueue)
 	{
 		QueuedRawPackets.Enqueue(PacketToQueue);
 	}
@@ -522,7 +522,7 @@ public:
 	 *
 	 * @param PacketToQueue		The packet to be queued
 	 */
-	FORCEINLINE void QueueHandlerPacketForSending(BufferedPacket* PacketToQueue)
+	inline void QueueHandlerPacketForSending(BufferedPacket* PacketToQueue)
 	{
 		QueuedHandlerPackets.Enqueue(PacketToQueue);
 	}
@@ -579,7 +579,7 @@ public:
 	 *
 	 * @param bInEnabled	Whether or not raw sends are enabled
 	 */
-	FORCEINLINE void SetRawSend(bool bInEnabled)
+	inline void SetRawSend(bool bInEnabled)
 	{
 		bRawSend = bInEnabled;
 	}
@@ -587,7 +587,7 @@ public:
 	/**
 	 * Whether or not raw packet sends are enabled
 	 */
-	FORCEINLINE bool GetRawSend()
+	inline bool GetRawSend()
 	{
 		return bRawSend;
 	}
@@ -595,7 +595,7 @@ public:
 	/**
 	 * Whether or not the packet handler is fully initialized, post-handshake etc.
 	 */
-	FORCEINLINE bool IsFullyInitialized()
+	inline bool IsFullyInitialized()
 	{
 		return State == UE::Handler::State::Initialized;
 	}
diff --git a/Engine/Source/Runtime/PacketHandlers/ReliabilityHandlerComponent/Public/ReliabilityHandlerComponent.h b/Engine/Source/Runtime/PacketHandlers/ReliabilityHandlerComponent/Public/ReliabilityHandlerComponent.h
--- a/Engine/Source/Runtime/PacketHandlers/ReliabilityHandlerComponent/Public/ReliabilityHandlerComponent.h
+++ b/Engine/Source/Runtime/PacketHandlers/ReliabilityHandlerComponent/Public/ReliabilityHandlerComponent.h
@@ -38,7 +38,7 @@ public:
 	RELIABILITYHANDLERCOMPONENT_API void QueuePacketForResending(uint8* Packet, int32 CountBits, FOutPacketTraits& Traits);
 
 	UE_DEPRECATED(4.21, "Use the PacketTraits version for sending packets with additional flags and options")
-	FORCEINLINE void QueueHandlerPacketForResending(HandlerComponent* InComponent, uint8* Packet, int32 CountBits)
+	inline void QueueHandlerPacketForResending(HandlerComponent* InComponent, uint8* Packet, int32 CountBits)
 	{
 		FOutPacketTraits EmptyTraits;
 		QueueHandlerPacketForResending(InComponent, Packet, CountBits, EmptyTraits);
@@ -51,7 +51,7 @@ public:
 	 * @param Packet		The packet to be queued
 	 * @param CountBits		The number of bits in the packet
 	 */
-	FORCEINLINE void QueueHandlerPacketForResending(HandlerComponent* InComponent, uint8* Packet, int32 CountBits, FOutPacketTraits& Traits)
+	inline void QueueHandlerPacketForResending(HandlerComponent* InComponent, uint8* Packet, int32 CountBits, FOutPacketTraits& Traits)
 	{
 		QueuePacketForResending(Packet, CountBits, Traits);
 
diff --git a/Engine/Source/Runtime/Sockets/Public/IPAddress.h b/Engine/Source/Runtime/Sockets/Public/IPAddress.h
--- a/Engine/Source/Runtime/Sockets/Public/IPAddress.h
+++ b/Engine/Source/Runtime/Sockets/Public/IPAddress.h
@@ -200,17 +200,17 @@ public:
 template<typename ValueType>
 struct FInternetAddrKeyMapFuncs : public BaseKeyFuncs<ValueType, TSharedRef<FInternetAddr>, false>
 {
-	static FORCEINLINE const TSharedRef<FInternetAddr>& GetSetKey(const TPair<TSharedRef<FInternetAddr>, ValueType>& Element)
+	static inline const TSharedRef<FInternetAddr>& GetSetKey(const TPair<TSharedRef<FInternetAddr>, ValueType>& Element)
 	{
 		return Element.Key;
 	}
 
-	static FORCEINLINE bool Matches(const TSharedRef<FInternetAddr>& A, const TSharedRef<FInternetAddr>& B)
+	static inline bool Matches(const TSharedRef<FInternetAddr>& A, const TSharedRef<FInternetAddr>& B)
 	{
 		return *A == *B;
 	}
 
-	static FORCEINLINE uint32 GetKeyHash(const TSharedRef<FInternetAddr>& Key)
+	static inline uint32 GetKeyHash(const TSharedRef<FInternetAddr>& Key)
 	{
 		return Key->GetTypeHash();
 	}
@@ -219,17 +219,17 @@ struct FInternetAddrKeyMapFuncs : public BaseKeyFuncs<ValueType, TSharedRef<FInt
 template<typename ValueType>
 struct FInternetAddrConstKeyMapFuncs : public BaseKeyFuncs<ValueType, TSharedRef<const FInternetAddr>, false>
 {
-	static FORCEINLINE const TSharedRef<const FInternetAddr>& GetSetKey(const TPair<TSharedRef<const FInternetAddr>, ValueType>& Element)
+	static inline const TSharedRef<const FInternetAddr>& GetSetKey(const TPair<TSharedRef<const FInternetAddr>, ValueType>& Element)
 	{
 		return Element.Key;
 	}
 
-	static FORCEINLINE bool Matches(const TSharedRef<const FInternetAddr>& A, const TSharedRef<const FInternetAddr>& B)
+	static inline bool Matches(const TSharedRef<const FInternetAddr>& A, const TSharedRef<const FInternetAddr>& B)
 	{
 		return *A == *B;
 	}
 
-	static FORCEINLINE uint32 GetKeyHash(const TSharedRef<const FInternetAddr>& Key)
+	static inline uint32 GetKeyHash(const TSharedRef<const FInternetAddr>& Key)
 	{
 		return Key->GetTypeHash();
 	}
diff --git a/Engine/Source/Runtime/Sockets/Public/IPAddressAsyncResolve.h b/Engine/Source/Runtime/Sockets/Public/IPAddressAsyncResolve.h
--- a/Engine/Source/Runtime/Sockets/Public/IPAddressAsyncResolve.h
+++ b/Engine/Source/Runtime/Sockets/Public/IPAddressAsyncResolve.h
@@ -126,7 +126,7 @@ class FResolveInfoAsync :
 			Parent->DoWork();
 		}
 
-		FORCEINLINE TStatId GetStatId() const
+		inline TStatId GetStatId() const
 		{
 			RETURN_QUICK_DECLARE_CYCLE_STAT(FResolveInfoAsyncWorker, STATGROUP_ThreadPoolAsyncTasks);
 		}
diff --git a/Engine/Source/Runtime/Sockets/Public/NetworkMessage.h b/Engine/Source/Runtime/Sockets/Public/NetworkMessage.h
--- a/Engine/Source/Runtime/Sockets/Public/NetworkMessage.h
+++ b/Engine/Source/Runtime/Sockets/Public/NetworkMessage.h
@@ -295,7 +295,7 @@ public:
 	}
 
 	// helper to serialize TCHAR* (there are a lot)
-	FORCEINLINE friend FNetworkFileArchive& operator<<(FNetworkFileArchive& Ar, const TCHAR*& Str)
+	inline friend FNetworkFileArchive& operator<<(FNetworkFileArchive& Ar, const TCHAR*& Str)
 	{
 		FString Temp(Str);
 		Ar << Temp;
diff --git a/Engine/Source/Runtime/Sockets/Public/Sockets.h b/Engine/Source/Runtime/Sockets/Public/Sockets.h
--- a/Engine/Source/Runtime/Sockets/Public/Sockets.h
+++ b/Engine/Source/Runtime/Sockets/Public/Sockets.h
@@ -447,7 +447,7 @@ public:
 	 * @return Socket type.
 	 * @see GetDescription, GetPortNo
 	 */
-	FORCEINLINE ESocketType GetSocketType() const
+	inline ESocketType GetSocketType() const
 	{
 		return SocketType;
 	}
@@ -458,7 +458,7 @@ public:
 	 * @return Debug string.
 	 * @see GetPortNo, GetSocketType
 	 */
-	FORCEINLINE FString GetDescription() const
+	inline FString GetDescription() const
 	{
 		return SocketDescription;
 	}
@@ -469,7 +469,7 @@ public:
 	 * @return Socket type.
 	 * @see GetDescription, GetPortNo
 	 */
-	FORCEINLINE FName GetProtocol() const
+	inline FName GetProtocol() const
 	{
 		return SocketProtocol;
 	}
