PR #13506: Code clean: cleanup usages that expect nullptr_t to be in the global namespace.

#13506

diff --git a/Engine/Plugins/Runtime/ReplicationGraph/Source/Public/ReplicationGraphTypes.h b/Engine/Plugins/Runtime/ReplicationGraph/Source/Public/ReplicationGraphTypes.h
--- a/Engine/Plugins/Runtime/ReplicationGraph/Source/Public/ReplicationGraphTypes.h
+++ b/Engine/Plugins/Runtime/ReplicationGraph/Source/Public/ReplicationGraphTypes.h
@@ -191,7 +191,7 @@ public:
 	}
 
 	// comparison with nullptr without resolving weak pointer or implicitly constructing a new FActorRepListType(nullptr).
-	friend bool operator==(const FActorRepListType& RepListActor, nullptr_t)
+	friend bool operator==(const FActorRepListType& RepListActor, TYPE_OF_NULLPTR)
 	{
 		if (const FWeakObjectPtr* WeakPtrPtr = RepListActor.ActorUnion.TryGet<FWeakObjectPtr>(); WeakPtrPtr)
 		{
@@ -201,15 +201,15 @@ public:
 		AActor* const* RawActorPtr = RepListActor.ActorUnion.TryGet<AActor*>();
 		return ensure(RawActorPtr) && *RawActorPtr == nullptr;
 	}
-	friend bool operator!=(const FActorRepListType& RepListActor, nullptr_t)
+	friend bool operator!=(const FActorRepListType& RepListActor, TYPE_OF_NULLPTR)
 	{
 		return !(RepListActor == nullptr);
 	}
-	friend bool operator==(nullptr_t, const FActorRepListType& RepListActor)
+	friend bool operator==(TYPE_OF_NULLPTR, const FActorRepListType& RepListActor)
 	{
 		return RepListActor == nullptr;
 	}
-	friend bool operator!=(nullptr_t, const FActorRepListType& RepListActor)
+	friend bool operator!=(TYPE_OF_NULLPTR, const FActorRepListType& RepListActor)
 	{
 		return !(RepListActor == nullptr);
 	}
diff --git a/Engine/Plugins/Runtime/ReplicationSystemTestPlugin/Source/Private/NetworkAutomationTestMacros.h b/Engine/Plugins/Runtime/ReplicationSystemTestPlugin/Source/Private/NetworkAutomationTestMacros.h
--- a/Engine/Plugins/Runtime/ReplicationSystemTestPlugin/Source/Private/NetworkAutomationTestMacros.h
+++ b/Engine/Plugins/Runtime/ReplicationSystemTestPlugin/Source/Private/NetworkAutomationTestMacros.h
@@ -107,7 +107,7 @@ Name(const T& Value, const T& ExpectedValue, const char* ValueText, const char*
 		return UE::Net::CreateTestFailure() << TEXT("Expected ") << ValueText << TEXT(' ') << TEXT(#Operator) <<  TEXT(' ') << ExpectedValueText << TEXT(". ") << Value << TEXT(" was compared to ") << ExpectedValue << TEXT(". "); \
 	} \
 } \
-template<typename T, typename U, typename TEnableIf<TIsPointer<T>::Value && std::is_same_v<U, nullptr_t>, U>::Type X = nullptr> \
+template<typename T, typename U, typename TEnableIf<TIsPointer<T>::Value && std::is_same_v<U, TYPE_OF_NULLPTR>, U>::Type X = nullptr> \
 UE::Net::FTestResult \
 Name(const T& Value, U ExpectedValue, const char* ValueText, const char* ExpectedValueText) \
 { \
diff --git a/Engine/Plugins/Runtime/ReplicationSystemTestPlugin/Source/Private/TestMessage.h b/Engine/Plugins/Runtime/ReplicationSystemTestPlugin/Source/Private/TestMessage.h
--- a/Engine/Plugins/Runtime/ReplicationSystemTestPlugin/Source/Private/TestMessage.h
+++ b/Engine/Plugins/Runtime/ReplicationSystemTestPlugin/Source/Private/TestMessage.h
@@ -117,7 +117,7 @@ inline FTestMessage& operator<<(FTestMessage& Message, bool Value)
 	return Message << (Value ? TEXT("true") : TEXT("false"));
 }
 
-inline FTestMessage& operator<<(FTestMessage& Message, nullptr_t)
+inline FTestMessage& operator<<(FTestMessage& Message, TYPE_OF_NULLPTR)
 {
 	return Message << TEXT("nullptr");
 }
diff --git a/Engine/Source/Editor/SequencerCore/Public/MVVM/ViewModelPtr.h b/Engine/Source/Editor/SequencerCore/Public/MVVM/ViewModelPtr.h
--- a/Engine/Source/Editor/SequencerCore/Public/MVVM/ViewModelPtr.h
+++ b/Engine/Source/Editor/SequencerCore/Public/MVVM/ViewModelPtr.h
@@ -155,11 +155,11 @@ struct TViewModelPtr : TViewModelConversions<T>
 	}
 
 	/** Construction from a nullptr */
-	TViewModelPtr(nullptr_t)
+	TViewModelPtr(TYPE_OF_NULLPTR)
 	{}
 
 	/** Assignment from a nullptr */
-	TViewModelPtr& operator=(nullptr_t)
+	TViewModelPtr& operator=(TYPE_OF_NULLPTR)
 	{
 		Storage.Set(nullptr, nullptr);
 		return *this;
@@ -318,11 +318,11 @@ public:
 		return A.Storage.SharedModel != B;
 	}
 
-	friend bool operator==(const TViewModelPtr<T>& A, nullptr_t B)
+	friend bool operator==(const TViewModelPtr<T>& A, TYPE_OF_NULLPTR B)
 	{
 		return A.Storage.SharedModel == nullptr;
 	}
-	friend bool operator!=(const TViewModelPtr<T>& A, nullptr_t B)
+	friend bool operator!=(const TViewModelPtr<T>& A, TYPE_OF_NULLPTR B)
 	{
 		return A.Storage.SharedModel != nullptr;
 	}
@@ -423,7 +423,7 @@ struct TWeakViewModelPtr : TWeakViewModelConversions<T>
 	TWeakViewModelPtr(TWeakViewModelPtr&&) = default;
 	TWeakViewModelPtr& operator=(TWeakViewModelPtr&&) = default;
 
-	TWeakViewModelPtr(nullptr_t)
+	TWeakViewModelPtr(TYPE_OF_NULLPTR)
 	{}
 
 	template<typename OtherType>
diff --git a/Engine/Source/Runtime/Core/Public/Serialization/CompactBinaryWriter.h b/Engine/Source/Runtime/Core/Public/Serialization/CompactBinaryWriter.h
--- a/Engine/Source/Runtime/Core/Public/Serialization/CompactBinaryWriter.h
+++ b/Engine/Source/Runtime/Core/Public/Serialization/CompactBinaryWriter.h
@@ -348,7 +348,7 @@ public:
 		return *this;
 	}
 
-	inline FCbWriter& operator<<(nullptr_t)
+	inline FCbWriter& operator<<(TYPE_OF_NULLPTR)
 	{
 		AddNull();
 		return *this;
diff --git a/Engine/Source/Runtime/MovieScene/Public/EntitySystem/RelativePtr.h b/Engine/Source/Runtime/MovieScene/Public/EntitySystem/RelativePtr.h
--- a/Engine/Source/Runtime/MovieScene/Public/EntitySystem/RelativePtr.h
+++ b/Engine/Source/Runtime/MovieScene/Public/EntitySystem/RelativePtr.h
@@ -17,7 +17,7 @@ struct TRelativePtr
 		: Offset(NullValue)
 	{}
 
-	explicit TRelativePtr(nullptr_t)
+	explicit TRelativePtr(TYPE_OF_NULLPTR)
 		: Offset(NullValue)
 	{}
 
@@ -42,7 +42,7 @@ struct TRelativePtr
 		return static_cast<U*>(Resolve(BasePtr));
 	}
 
-	void Reset(nullptr_t)
+	void Reset(TYPE_OF_NULLPTR)
 	{
 		Offset = NullValue;
 	}
diff --git a/Engine/Source/Runtime/PakFile/Public/IPlatformFilePak.h b/Engine/Source/Runtime/PakFile/Public/IPlatformFilePak.h
--- a/Engine/Source/Runtime/PakFile/Public/IPlatformFilePak.h
+++ b/Engine/Source/Runtime/PakFile/Public/IPlatformFilePak.h
@@ -711,8 +711,8 @@ public:
 	PAKFILE_API FSharedPakReader& operator=(FSharedPakReader&& Other);
 
 	explicit operator bool() const { return Archive != nullptr; }
-	bool operator==(nullptr_t) { return Archive == nullptr; }
-	bool operator!=(nullptr_t) { return Archive != nullptr; }
+	bool operator==(TYPE_OF_NULLPTR) { return Archive == nullptr; }
+	bool operator!=(TYPE_OF_NULLPTR) { return Archive != nullptr; }
 	FArchive* operator->() { return Archive; }
 
 	
diff --git a/Engine/Source/Runtime/RenderCore/Public/VertexFactory.h b/Engine/Source/Runtime/RenderCore/Public/VertexFactory.h
--- a/Engine/Source/Runtime/RenderCore/Public/VertexFactory.h
+++ b/Engine/Source/Runtime/RenderCore/Public/VertexFactory.h
@@ -72,7 +72,7 @@ struct FVertexInputStream
 		VertexBuffer(nullptr)
 	{}
 	
-	FVertexInputStream(uint32 InStreamIndex, uint32 InOffset, nullptr_t)
+	FVertexInputStream(uint32 InStreamIndex, uint32 InOffset, TYPE_OF_NULLPTR)
 		: StreamIndex(InStreamIndex), Offset(InOffset), bStreamSourceSlot(0), VertexBuffer(nullptr)
 	{
 		// Verify no overflow
